Set supported target platforms on video production plugin to Win64

diff --git a/Engine/Plugins/Media/AppleProResMedia/AppleProResMedia.uplugin b/Engine/Plugins/Media/AppleProResMedia/AppleProResMedia.uplugin
--- a/Engine/Plugins/Media/AppleProResMedia/AppleProResMedia.uplugin
+++ b/Engine/Plugins/Media/AppleProResMedia/AppleProResMedia.uplugin
@@ -1,18 +1,18 @@
 {
-	"FileVersion" : 3,
-	"Version" : 1,
-	"VersionName" : "1.0",
-	"FriendlyName" : "Apple ProRes Media",
-	"Description" : "Implements video playback and the export of the Apple ProRes Codec.  Apple ProRes is a high quality, lossy video compression format.",
-	"CreatedBy" : "Epic Games Inc",
-	"CreatedByURL" : "http://epicgames.com",
-	"Category" : "Media Players",
+	"FileVersion": 3,
+	"Version": 1,
+	"VersionName": "1.0",
+	"FriendlyName": "Apple ProRes Media",
+	"Description": "Implements video playback and the export of the Apple ProRes Codec.  Apple ProRes is a high quality, lossy video compression format.",
+	"CreatedBy": "Epic Games Inc",
+	"CreatedByURL": "http://epicgames.com",
+	"Category": "Media Players",
 	"EnabledByDefault": false,
-	"CanContainContent" : false,
+	"CanContainContent": false,
 	"IsBetaVersion": false,
 	"Installed": false,
-	"Modules" :
-	[
+	"SupportedTargetPlatforms": [ "Win64" ],
+	"Modules": [
 		{
 			"Name": "AppleProResMedia",
 			"Type": "Runtime",
diff --git a/Engine/Plugins/Media/HAPMedia/HAPMedia.uplugin b/Engine/Plugins/Media/HAPMedia/HAPMedia.uplugin
--- a/Engine/Plugins/Media/HAPMedia/HAPMedia.uplugin
+++ b/Engine/Plugins/Media/HAPMedia/HAPMedia.uplugin
@@ -11,6 +11,7 @@
 	"CanContainContent": false,
 	"IsBetaVersion": false,
 	"Installed": false,
+	"SupportedTargetPlatforms": [ "Win64" ],
 	"Modules": [
 		{
 			"Name": "HAPMedia",
diff --git a/Engine/Plugins/Runtime/nDisplay/nDisplay.uplugin b/Engine/Plugins/Runtime/nDisplay/nDisplay.uplugin
--- a/Engine/Plugins/Runtime/nDisplay/nDisplay.uplugin
+++ b/Engine/Plugins/Runtime/nDisplay/nDisplay.uplugin
@@ -14,6 +14,7 @@
 	"CanContainContent": true,
 	"IsBetaVersion": true,
 	"Installed": false,
+	"SupportedTargetPlatforms": [ "Win64" ],
 	"Modules": [
 		{
 			"Name": "DisplayCluster",
