ProRes/Dnx: Defer file creation until after sequence has updated

diff --git a/Engine/Plugins/Media/AppleProResMedia/Source/AppleProResMedia/Private/AppleProResEncoderProtocol.cpp b/Engine/Plugins/Media/AppleProResMedia/Source/AppleProResMedia/Private/AppleProResEncoderProtocol.cpp
--- a/Engine/Plugins/Media/AppleProResMedia/Source/AppleProResMedia/Private/AppleProResEncoderProtocol.cpp
+++ b/Engine/Plugins/Media/AppleProResMedia/Source/AppleProResMedia/Private/AppleProResEncoderProtocol.cpp
@@ -34,11 +34,6 @@ bool UAppleProResEncoderProtocol::SetupImpl()
 {
 	ParseCommandLine();
 
-	if (!CreateFile(GenerateFilenameImpl(FFrameMetrics(), TEXT(".mov"))))
-	{
-		return false;
-	}
-
 	return Super::SetupImpl();
 }
 
@@ -224,6 +219,15 @@ bool UAppleProResEncoderProtocol::CreateFile(const FString& InFilename)
 
 void UAppleProResEncoderProtocol::ProcessFrame(FCapturedFrameData InFrame)
 {
+	// Defer writer creation until sequence has evaluated so that cached data for format mappings can be generated
+	if (FileWriter == nullptr)
+	{
+		if (!CreateFile(GenerateFilenameImpl(FFrameMetrics(), TEXT(".mov"))))
+		{
+			return;
+		}
+	}
+
 	FIntPoint ImageSize = FIntPoint(Dimensions.width, Dimensions.height);
 
 	FVideoFrameData* Payload = InFrame.GetPayload<FVideoFrameData>();
diff --git a/Engine/Plugins/Media/AvidDNxMedia/Source/Source/Private/AvidDNxEncoderProtocol.cpp b/Engine/Plugins/Media/AvidDNxMedia/Source/Source/Private/AvidDNxEncoderProtocol.cpp
--- a/Engine/Plugins/Media/AvidDNxMedia/Source/Source/Private/AvidDNxEncoderProtocol.cpp
+++ b/Engine/Plugins/Media/AvidDNxMedia/Source/Source/Private/AvidDNxEncoderProtocol.cpp
@@ -109,12 +109,6 @@ bool UAvidDNxEncoderProtocol::SetupImpl()
 		return false;
 	}
 
-	if (!InitializeMXFwriter())
-	{
-		DestroyEncoder();
-		return false;
-	}
-
 	return Super::SetupImpl();
 }
 
@@ -301,6 +295,15 @@ FFramePayloadPtr UAvidDNxEncoderProtocol::GetFramePayload(const FFrameMetrics& I
 
 void UAvidDNxEncoderProtocol::ProcessFrame(FCapturedFrameData InFrame)
 {
+	// Defer writer creation until sequence has evaluated so that cached data for format mappings can be generated
+	if (MXFwriter == nullptr)
+	{
+		if (!InitializeMXFwriter())
+		{
+			return;
+		}
+	}
+
 	const double StartTime = FPlatformTime::Seconds();
 
 	ensure(InFrame.ColorBuffer.Num() % 2 == 0);
