Fixed AppleProResMedia compile issues. These appeared after doing some compile time optimizations.

diff --git a/Engine/Plugins/Media/AppleProResMedia/Source/AppleProResMedia/Private/AppleProResEncoderProtocol.cpp b/Engine/Plugins/Media/AppleProResMedia/Source/AppleProResMedia/Private/AppleProResEncoderProtocol.cpp
--- a/Engine/Plugins/Media/AppleProResMedia/Source/AppleProResMedia/Private/AppleProResEncoderProtocol.cpp
+++ b/Engine/Plugins/Media/AppleProResMedia/Source/AppleProResMedia/Private/AppleProResEncoderProtocol.cpp
@@ -88,7 +88,7 @@ void UAppleProResEncoderProtocol::ParseCommandLine()
 	}
 }
 
-bool UAppleProResEncoderProtocol::CreateFile(const FString& InFilename)
+bool UAppleProResEncoderProtocol::CreateProResFile(const FString& InFilename)
 {
 	check(FileWriter == nullptr);
 
@@ -222,7 +222,7 @@ void UAppleProResEncoderProtocol::ProcessFrame(FCapturedFrameData InFrame)
 	// Defer writer creation until sequence has evaluated so that cached data for format mappings can be generated
 	if (FileWriter == nullptr)
 	{
-		if (!CreateFile(GenerateFilenameImpl(FFrameMetrics(), TEXT(".mov"))))
+		if (!CreateProResFile(GenerateFilenameImpl(FFrameMetrics(), TEXT(".mov"))))
 		{
 			return;
 		}
diff --git a/Engine/Plugins/Media/AppleProResMedia/Source/AppleProResMedia/Private/AppleProResEncoderProtocol.h b/Engine/Plugins/Media/AppleProResMedia/Source/AppleProResMedia/Private/AppleProResEncoderProtocol.h
--- a/Engine/Plugins/Media/AppleProResMedia/Source/AppleProResMedia/Private/AppleProResEncoderProtocol.h
+++ b/Engine/Plugins/Media/AppleProResMedia/Source/AppleProResMedia/Private/AppleProResEncoderProtocol.h
@@ -84,7 +84,7 @@ public:
 private:
 
 	void ParseCommandLine();
-	bool CreateFile(const FString& InFilename);
+	bool CreateProResFile(const FString& InFilename);
 	void ConvertFColorToRGBA4444(const TArray<FColor>& InColorbuffer);
 
 	PRCodecType GetSelectedCodecType() const;
diff --git a/Engine/Plugins/Media/AppleProResMedia/Source/AppleProResMedia/Private/MoviePipelineAppleProResOutput.h b/Engine/Plugins/Media/AppleProResMedia/Source/AppleProResMedia/Private/MoviePipelineAppleProResOutput.h
--- a/Engine/Plugins/Media/AppleProResMedia/Source/AppleProResMedia/Private/MoviePipelineAppleProResOutput.h
+++ b/Engine/Plugins/Media/AppleProResMedia/Source/AppleProResMedia/Private/MoviePipelineAppleProResOutput.h
@@ -1,5 +1,7 @@
 // Copyright Epic Games, Inc. All Rights Reserved.
 
+#pragma once
+
 #include "MoviePipelineVideoOutputBase.h"
 #include "AppleProResEncoder/AppleProResEncoder.h"
 #include "MoviePipelineAppleProResOutput.generated.h"
