Renaming member variable to be more inclusive

#rnx

diff --git a/Engine/Plugins/Media/AppleProResMedia/Source/AppleProResMedia/Private/AppleProResEncoder/AppleProResEncoder.cpp b/Engine/Plugins/Media/AppleProResMedia/Source/AppleProResMedia/Private/AppleProResEncoder/AppleProResEncoder.cpp
--- a/Engine/Plugins/Media/AppleProResMedia/Source/AppleProResMedia/Private/AppleProResEncoder/AppleProResEncoder.cpp
+++ b/Engine/Plugins/Media/AppleProResMedia/Source/AppleProResMedia/Private/AppleProResEncoder/AppleProResEncoder.cpp
@@ -661,7 +661,7 @@ bool FAppleProResEncoder::WriteFrame(const FImagePixelData* InPixelData)
 	if (TimecodePayload)
 	{
 		uint32 Timecode;
-		Timecode = _byteswap_ulong(TimecodePayload->MasterFrameNumber);
+		Timecode = _byteswap_ulong(TimecodePayload->ReferenceFrameNumber);
 
 		SampleSize = sizeof(Timecode);
 		Status = ProResFileWriterAddSampleBufferToTrack(
diff --git a/Engine/Plugins/Media/AppleProResMedia/Source/AppleProResMedia/Private/MoviePipelineAppleProResOutput.cpp b/Engine/Plugins/Media/AppleProResMedia/Source/AppleProResMedia/Private/MoviePipelineAppleProResOutput.cpp
--- a/Engine/Plugins/Media/AppleProResMedia/Source/AppleProResMedia/Private/MoviePipelineAppleProResOutput.cpp
+++ b/Engine/Plugins/Media/AppleProResMedia/Source/AppleProResMedia/Private/MoviePipelineAppleProResOutput.cpp
@@ -63,7 +63,7 @@ void UMoviePipelineAppleProResOutput::WriteFrame_EncodeThread(MovieRenderPipelin
 	TSharedPtr<FAppleProResEncoder::FTimecodePayload, ESPMode::ThreadSafe> ProResPayload = MakeShared<FAppleProResEncoder::FTimecodePayload, ESPMode::ThreadSafe>();
 
 	// This is the frame number on the global time, can have overlaps (between encoders) or repeats when using handle frames/slowmo.
-	ProResPayload->MasterFrameNumber = PipelinePayload->SampleState.OutputState.SourceFrameNumber;
+	ProResPayload->ReferenceFrameNumber = PipelinePayload->SampleState.OutputState.SourceFrameNumber;
 
 	// ProRes can handle quantization internally but expects sRGB to be applied to the incoming data.
 	TUniquePtr<FImagePixelData> QuantizedPixelData = UE::MoviePipeline::QuantizeImagePixelDataToBitDepth(InPixelData, 16, ProResPayload, InWriter->bConvertToSrgb);
diff --git a/Engine/Plugins/Media/AppleProResMedia/Source/AppleProResMedia/Public/AppleProResEncoder/AppleProResEncoder.h b/Engine/Plugins/Media/AppleProResMedia/Source/AppleProResMedia/Public/AppleProResEncoder/AppleProResEncoder.h
--- a/Engine/Plugins/Media/AppleProResMedia/Source/AppleProResMedia/Public/AppleProResEncoder/AppleProResEncoder.h
+++ b/Engine/Plugins/Media/AppleProResMedia/Source/AppleProResMedia/Public/AppleProResEncoder/AppleProResEncoder.h
@@ -107,7 +107,7 @@ class APPLEPRORESMEDIA_API FAppleProResEncoder
 public:
 	struct FTimecodePayload : IImagePixelDataPayload
 	{
-		int32 MasterFrameNumber;
+		int32 ReferenceFrameNumber;
 	};
 
 public:
