MoviePipeline: Added an Apple ProRes Movie node to the graph.

Also fixed a bug in the base UMovieGraphVideoOutputNode that caused the File Name Format property to be ignored (and only used the CDO's value).

diff --git a/Engine/Plugins/Media/AppleProResMedia/Source/AppleProResMedia/AppleProResMedia.Build.cs b/Engine/Plugins/Media/AppleProResMedia/Source/AppleProResMedia/AppleProResMedia.Build.cs
--- a/Engine/Plugins/Media/AppleProResMedia/Source/AppleProResMedia/AppleProResMedia.Build.cs
+++ b/Engine/Plugins/Media/AppleProResMedia/Source/AppleProResMedia/AppleProResMedia.Build.cs
@@ -14,6 +14,8 @@ namespace UnrealBuildTool.Rules
                     "WmfMedia",
 					"ImageWriteQueue",
 					"MovieRenderPipelineCore",
+					"OpenColorIO",
+					"SlateCore",
                 });
 
 			PrivateDependencyModuleNames.AddRange(
diff --git a/Engine/Plugins/Media/AppleProResMedia/Source/AppleProResMedia/Private/AppleProResEncoder/AppleProResEncoder.cpp b/Engine/Plugins/Media/AppleProResMedia/Source/AppleProResMedia/Private/AppleProResEncoder/AppleProResEncoder.cpp
--- a/Engine/Plugins/Media/AppleProResMedia/Source/AppleProResMedia/Private/AppleProResEncoder/AppleProResEncoder.cpp
+++ b/Engine/Plugins/Media/AppleProResMedia/Source/AppleProResMedia/Private/AppleProResEncoder/AppleProResEncoder.cpp
@@ -692,3 +692,8 @@ bool FAppleProResEncoder::WriteAudioSample(const TArrayView<int16>& InAudioSampl
 {
 	return false;
 }
+
+const FAppleProResEncoderOptions& FAppleProResEncoder::GetOptions() const
+{
+	return Options;
+}
diff --git a/Engine/Plugins/Media/AppleProResMedia/Source/AppleProResMedia/Private/AppleProResMediaModule.cpp b/Engine/Plugins/Media/AppleProResMedia/Source/AppleProResMedia/Private/AppleProResMediaModule.cpp
--- a/Engine/Plugins/Media/AppleProResMedia/Source/AppleProResMedia/Private/AppleProResMediaModule.cpp
+++ b/Engine/Plugins/Media/AppleProResMedia/Source/AppleProResMedia/Private/AppleProResMediaModule.cpp
@@ -77,6 +77,9 @@ class FAppleProResMediaModule : public IModuleInterface
 			);
 		}
 #endif //WITH_EDITOR
+
+		// Add exemption to FName::NameToDisplayString formatting to ensure "ProRes" is displayed without a space
+		FName::AddNameToDisplayStringExemption(TEXT("ProRes"));
 	}
 
 	virtual void ShutdownModule() override
diff --git a/Engine/Plugins/Media/AppleProResMedia/Source/AppleProResMedia/Private/MovieGraphAppleProResNode.cpp b/Engine/Plugins/Media/AppleProResMedia/Source/AppleProResMedia/Private/MovieGraphAppleProResNode.cpp
new file mode 100644
--- /dev/null
+++ b/Engine/Plugins/Media/AppleProResMedia/Source/AppleProResMedia/Private/MovieGraphAppleProResNode.cpp
@@ -0,0 +1,199 @@
+// Copyright Epic Games, Inc. All Rights Reserved.
+
+#include "MovieGraphAppleProResNode.h"
+
+#include "AppleProResEncoder/AppleProResEncoder.h"
+#include "Graph/MovieGraphBlueprintLibrary.h"
+#include "Graph/MovieGraphConfig.h"
+#include "Graph/MovieGraphPipeline.h"
+#include "Graph/MovieGraphOCIOHelper.h"
+#include "Graph/Nodes/MovieGraphGlobalOutputSettingNode.h"
+#include "ImageWriteTask.h"
+#include "MoviePipelineImageQuantization.h"
+#include "MovieRenderPipelineCoreModule.h"
+#include "Styling/AppStyle.h"
+
+UMovieGraphAppleProResNode::UMovieGraphAppleProResNode()
+	: Quality(EAppleProResEncoderCodec::ProRes_422LT)
+	, bDropFrameTimecode(true)	// Defaults to true because most 29.97 FPS content uses this
+{
+	
+}
+
+EMovieGraphBranchRestriction UMovieGraphAppleProResNode::GetBranchRestriction() const
+{
+	return EMovieGraphBranchRestriction::Globals;
+}
+
+#if WITH_EDITOR
+FText UMovieGraphAppleProResNode::GetNodeTitle(const bool bGetDescriptive) const
+{
+	static const FText AppleProResNodeName = NSLOCTEXT("MovieGraphNodes", "NodeName_AppleProRes", "Apple ProRes Movie");
+	return AppleProResNodeName;
+}
+
+FText UMovieGraphAppleProResNode::GetMenuCategory() const
+{
+	return NSLOCTEXT("MovieGraphNodes", "AppleProResNode_Category", "Output Type");
+}
+
+FText UMovieGraphAppleProResNode::GetKeywords() const
+{
+	static const FText Keywords = NSLOCTEXT("MovieGraphNodes", "AppleProResGraphNode_Keywords", "apple pro res prores mov movie video");
+	return Keywords;
+}
+
+FLinearColor UMovieGraphAppleProResNode::GetNodeTitleColor() const
+{
+	static const FLinearColor AppleProResNodeColor = FLinearColor(0.047f, 0.654f, 0.537f);
+	return AppleProResNodeColor;
+}
+
+FSlateIcon UMovieGraphAppleProResNode::GetIconAndTint(FLinearColor& OutColor) const
+{
+	static const FSlateIcon AppleProResIcon = FSlateIcon(FAppStyle::GetAppStyleSetName(), "LevelEditor.Tabs.Cinematics");
+
+	OutColor = FLinearColor::White;
+	return AppleProResIcon;
+}
+#endif // WITH_EDITOR
+
+TUniquePtr<MovieRenderGraph::IVideoCodecWriter> UMovieGraphAppleProResNode::Initialize_GameThread(UMovieGraphPipeline* InPipeline, TObjectPtr<UMovieGraphEvaluatedConfig> InEvaluatedConfig, const FString& InBranchName, const FString& InFileName, FIntPoint InResolution, EImagePixelType InPixelType, ERGBFormat InPixelFormat, uint8 InBitDepth, uint8 InNumChannels, bool bAllowOCIO)
+{
+	bool bIncludeCDOs = true;
+	constexpr bool bExactMatch = true;
+	UMovieGraphGlobalOutputSettingNode* OutputSetting =
+		InEvaluatedConfig->GetSettingForBranch<UMovieGraphGlobalOutputSettingNode>(GlobalsPinName, bIncludeCDOs, bExactMatch);
+
+	bIncludeCDOs = false;
+	const UMovieGraphAppleProResNode* EvaluatedNode = Cast<UMovieGraphAppleProResNode>(
+		InEvaluatedConfig->GetSettingForBranch(GetClass(), FName(InBranchName), bIncludeCDOs, bExactMatch));
+	checkf(EvaluatedNode, TEXT("Apple ProRes node could not be found in the graph in branch [%s]."), *InBranchName);
+	
+	const FFrameRate SourceFrameRate = InPipeline->GetDataSourceInstance()->GetDisplayRate();
+	const FFrameRate EffectiveFrameRate = UMovieGraphBlueprintLibrary::GetEffectiveFrameRate(OutputSetting, SourceFrameRate);
+	
+	FAppleProResEncoderOptions Options;
+	Options.OutputFilename = InFileName;
+	Options.Width = InResolution.X;
+	Options.Height = InResolution.Y;
+	Options.FrameRate = EffectiveFrameRate;
+	Options.Codec = EvaluatedNode->Quality;
+	Options.ColorPrimaries = EAppleProResEncoderColorPrimaries::CD_HDREC709; // Force Rec 709 for now
+	Options.ScanMode = EAppleProResEncoderScanMode::IM_PROGRESSIVE_SCAN; // No interlace sources
+	Options.bWriteAlpha = true;
+	Options.bDropFrameTimecode = bDropFrameTimecode;
+
+	// If OCIO is enabled, don't do additional color conversion
+	Options.bConvertToSrgb = !(EvaluatedNode->bOverride_OCIOConfiguration && EvaluatedNode->OCIOConfiguration.bIsEnabled && bAllowOCIO);
+	
+	TUniquePtr<FProResWriter> NewWriter = MakeUnique<FProResWriter>();
+	NewWriter->Writer = MakeUnique<FAppleProResEncoder>(Options);
+
+	CachedPipeline = InPipeline;
+	
+	return NewWriter;
+}
+
+bool UMovieGraphAppleProResNode::Initialize_EncodeThread(MovieRenderGraph::IVideoCodecWriter* InWriter)
+{
+	const FProResWriter* CodecWriter = static_cast<FProResWriter*>(InWriter);
+	if(!CodecWriter->Writer->Initialize())
+	{
+		UE_LOG(LogMovieRenderPipeline, Error, TEXT("Failed to initialize Apple ProRes writer."));
+		return false;
+	}
+	
+	return true;
+}
+
+void UMovieGraphAppleProResNode::WriteFrame_EncodeThread(MovieRenderGraph::IVideoCodecWriter* InWriter, FImagePixelData* InPixelData, TArray<FMovieGraphPassData>&& InCompositePasses, TObjectPtr<UMovieGraphEvaluatedConfig> InEvaluatedConfig, const FString& InBranchName)
+{
+	constexpr bool bIncludeCDOs = false;
+	constexpr bool bExactMatch = true;
+	const UMovieGraphAppleProResNode* EvaluatedNode = Cast<UMovieGraphAppleProResNode>(
+		InEvaluatedConfig->GetSettingForBranch(GetClass(), FName(InBranchName), bIncludeCDOs, bExactMatch));
+	checkf(EvaluatedNode, TEXT("Apple ProRes node could not be found in the graph in branch [%s]."), *InBranchName);
+	
+	const FProResWriter* CodecWriter = static_cast<FProResWriter*>(InWriter);
+	
+	const UE::MovieGraph::FMovieGraphSampleState* GraphPayload = InPixelData->GetPayload<UE::MovieGraph::FMovieGraphSampleState>();
+	
+	// Translate our Movie Pipeline specific payload to a ProRes Encoder specific payload.
+	const TSharedPtr<FAppleProResEncoder::FTimecodePayload, ESPMode::ThreadSafe> ProResPayload = MakeShared<FAppleProResEncoder::FTimecodePayload, ESPMode::ThreadSafe>();
+
+	if (EvaluatedNode->bOverride_CustomTimecodeStart)
+	{
+		const int32 CustomTimecodeStartFrame = EvaluatedNode->CustomTimecodeStart.ToFrameNumber(CodecWriter->Writer->GetOptions().FrameRate).Value;
+		const FMovieGraphTimeStepData& TimeData = GraphPayload->TraversalContext.Time;
+		
+		// When using a custom timecode start, just use the root-level frame number (relative to zero) offset by the custom timecode start
+		ProResPayload->ReferenceFrameNumber = TimeData.OutputFrameNumber + CustomTimecodeStartFrame;
+	}
+	else
+	{
+		// This is the frame number on the global time, can have overlaps (between encoders) or repeats when using handle frames/slowmo.
+		ProResPayload->ReferenceFrameNumber = GraphPayload->TraversalContext.Time.RootFrameNumber.Value;
+	}
+
+	// ProRes can handle quantization internally but expects sRGB to be applied to the incoming data.
+	const TUniquePtr<FImagePixelData> QuantizedPixelData = UE::MoviePipeline::QuantizeImagePixelDataToBitDepth(InPixelData, 8, ProResPayload, CodecWriter->Writer->GetOptions().bConvertToSrgb);
+
+	TArray<FPixelPreProcessor> PixelPreProcessors;
+
+#if WITH_OCIO
+	FMovieGraphOCIOHelper::GenerateOcioPixelPreProcessor(GraphPayload, CachedPipeline.Get(), InEvaluatedConfig, OCIOConfiguration, OCIOContext, PixelPreProcessors);
+#endif
+
+	// Do a quick composite of renders/burn-ins.
+	for (const FMovieGraphPassData& CompositePass : InCompositePasses)
+	{
+		// We don't need to copy the data here (even though it's being passed to a async system) because we already made a unique copy of the
+		// burn in/widget data when we decided to composite it.
+		switch (QuantizedPixelData->GetType())
+		{
+		case EImagePixelType::Color:
+			PixelPreProcessors.Add(TAsyncCompositeImage<FColor>(CompositePass.Value->MoveImageDataToNew()));
+			break;
+		case EImagePixelType::Float16:
+			PixelPreProcessors.Add(TAsyncCompositeImage<FFloat16Color>(CompositePass.Value->MoveImageDataToNew()));
+			break;
+		case EImagePixelType::Float32:
+			PixelPreProcessors.Add(TAsyncCompositeImage<FLinearColor>(CompositePass.Value->MoveImageDataToNew()));
+			break;
+		}
+	}
+
+	// This is done on the current thread for simplicity but the composite itself is parallelized.
+	FImagePixelData* PixelData = QuantizedPixelData.Get();
+	for (const FPixelPreProcessor& PreProcessor : PixelPreProcessors)
+	{
+		// PreProcessors are assumed to be valid.
+		PreProcessor(PixelData);
+	}
+
+	CodecWriter->Writer->WriteFrame(PixelData);
+}
+
+void UMovieGraphAppleProResNode::BeginFinalize_EncodeThread(MovieRenderGraph::IVideoCodecWriter* InWriter)
+{
+	return;
+}
+
+void UMovieGraphAppleProResNode::Finalize_EncodeThread(MovieRenderGraph::IVideoCodecWriter* InWriter)
+{
+	// Write to disk
+	const FProResWriter* CodecWriter = static_cast<FProResWriter*>(InWriter);
+	CodecWriter->Writer->Finalize();
+}
+
+const TCHAR* UMovieGraphAppleProResNode::GetFilenameExtension() const
+{
+	return TEXT("mov");
+}
+
+bool UMovieGraphAppleProResNode::IsAudioSupported() const
+{
+	// The current ProRes SDK does not support audio so we don't write audio.
+	return false;
+}
\ No newline at end of file
diff --git a/Engine/Plugins/Media/AppleProResMedia/Source/AppleProResMedia/Private/MovieGraphAppleProResNode.h b/Engine/Plugins/Media/AppleProResMedia/Source/AppleProResMedia/Private/MovieGraphAppleProResNode.h
new file mode 100644
--- /dev/null
+++ b/Engine/Plugins/Media/AppleProResMedia/Source/AppleProResMedia/Private/MovieGraphAppleProResNode.h
@@ -0,0 +1,98 @@
+﻿// Copyright Epic Games, Inc. All Rights Reserved.
+
+#pragma once
+
+#include "AppleProResEncoder/AppleProResEncoder.h"
+#include "Graph/MovieGraphNode.h"
+#include "Graph/Nodes/MovieGraphVideoOutputNode.h"
+
+#include "MovieGraphAppleProResNode.generated.h"
+
+/** A node which can output Apple ProRes movies. */
+UCLASS(BlueprintType, PrioritizeCategories=("FileOutput"))
+class UMovieGraphAppleProResNode : public UMovieGraphVideoOutputNode
+{
+	GENERATED_BODY()
+
+public:
+	UMovieGraphAppleProResNode();
+
+	virtual EMovieGraphBranchRestriction GetBranchRestriction() const override;
+
+#if WITH_EDITOR
+	virtual FText GetNodeTitle(const bool bGetDescriptive = false) const override;
+	virtual FText GetMenuCategory() const override;
+	virtual FText GetKeywords() const override;
+	virtual FLinearColor GetNodeTitleColor() const override;
+	virtual FSlateIcon GetIconAndTint(FLinearColor& OutColor) const override;
+#endif
+
+protected:
+	// UMovieGraphVideoOutputNode Interface
+	virtual TUniquePtr<MovieRenderGraph::IVideoCodecWriter> Initialize_GameThread(UMovieGraphPipeline* InPipeline, TObjectPtr<UMovieGraphEvaluatedConfig> InEvaluatedConfig, const FString& InBranchName, const FString& InFileName, FIntPoint InResolution, EImagePixelType InPixelType, ERGBFormat InPixelFormat, uint8 InBitDepth, uint8 InNumChannels, bool bAllowOCIO) override;
+	virtual bool Initialize_EncodeThread(MovieRenderGraph::IVideoCodecWriter* InWriter) override;
+	virtual void WriteFrame_EncodeThread(MovieRenderGraph::IVideoCodecWriter* InWriter, FImagePixelData* InPixelData, TArray<FMovieGraphPassData>&& InCompositePasses, TObjectPtr<UMovieGraphEvaluatedConfig> InEvaluatedConfig, const FString& InBranchName) override;
+	virtual void BeginFinalize_EncodeThread(MovieRenderGraph::IVideoCodecWriter* InWriter) override;
+	virtual void Finalize_EncodeThread(MovieRenderGraph::IVideoCodecWriter* InWriter) override;
+	virtual const TCHAR* GetFilenameExtension() const override;
+	virtual bool IsAudioSupported() const override;
+	// ~UMovieGraphVideoOutputNode Interface
+
+protected:
+	struct FProResWriter : public MovieRenderGraph::IVideoCodecWriter
+	{
+		TUniquePtr<FAppleProResEncoder> Writer;
+	};
+	
+	/** The pipeline that is running this node. */
+	TWeakObjectPtr<UMovieGraphPipeline> CachedPipeline;
+
+public:
+	UPROPERTY(EditAnywhere, BlueprintReadWrite, Category = Overrides, meta = (InlineEditConditionToggle))
+	uint8 bOverride_Quality : 1;
+
+	UPROPERTY(EditAnywhere, BlueprintReadWrite, Category = Overrides, meta = (InlineEditConditionToggle))
+	uint8 bOverride_CustomTimecodeStart : 1;
+
+	UPROPERTY(EditAnywhere, BlueprintReadWrite, Category = Overrides, meta = (InlineEditConditionToggle))
+	uint8 bOverride_bDropFrameTimecode : 1;
+
+	UPROPERTY(EditAnywhere, BlueprintReadWrite, Category = Overrides, meta = (InlineEditConditionToggle))
+	uint8 bOverride_OCIOConfiguration : 1;
+
+	UPROPERTY(EditAnywhere, BlueprintReadWrite, Category = Overrides, meta = (InlineEditConditionToggle))
+	uint8 bOverride_OCIOContext : 1;
+	
+	/** The Apple ProRes codec that should be used. See Apple documentation for more specifics. Uses Rec 709 color primaries. */
+	UPROPERTY(EditAnywhere, BlueprintReadWrite, Category = "Apple ProRes", meta = (EditCondition = "bOverride_Quality"))
+	EAppleProResEncoderCodec Quality;
+
+	/** Start the timecode at a specific value, rather than the value coming from the Level Sequence. */
+	UPROPERTY(EditAnywhere, BlueprintReadWrite, Category = "Apple ProRes", meta = (EditCondition = "bOverride_CustomTimecodeStart"))
+	FTimecode CustomTimecodeStart;
+
+	/** Whether the embedded timecode track should be written using drop-frame format. Only applicable if the sequence framerate is 29.97 */
+	UPROPERTY(EditAnywhere, BlueprintReadWrite, Category = "Apple ProRes", DisplayName = "Use DF Timecode if 29.97 FPS", meta = (EditCondition = "bOverride_bDropFrameTimecode"))
+	bool bDropFrameTimecode;
+
+	/**
+	* OCIO configuration/transform settings.
+	*
+	* Note: There are differences from the previous implementation in MRQ given that we are now doing CPU-side processing.
+	* 1) This feature only works on desktop platforms when the OpenColorIO library is available.
+	* 2) Users are now responsible for setting the renderer output space to Final Color (HDR) in Linear Working Color Space (SCS_FinalColorHDR) by
+	*    disabling the Tone Curve setting on the renderer node.
+	*/
+	UPROPERTY(EditAnywhere, BlueprintReadWrite, Category = "OCIO", DisplayName="OCIO Configuration", meta = (DisplayAfter = "FileNameFormat", EditCondition = "bOverride_OCIOConfiguration"))
+	FOpenColorIODisplayConfiguration OCIOConfiguration;
+
+	/**
+	* OCIO context of key-value string pairs, typically used to apply shot-specific looks (such as a CDL color correction, or a 1D grade LUT).
+	* 
+	* Notes:
+	* 1) If a configuration asset base context was set, it remains active but can be overridden here with new key-values.
+	* 2) Format tokens such as {shot_name} are supported and will get resolved before submission.
+	*/
+	UPROPERTY(EditAnywhere, BlueprintReadWrite, Category = "OCIO", DisplayName = "OCIO Context", meta = (DisplayAfter = "OCIOConfiguration", EditCondition = "bOverride_OCIOContext"))
+	TMap<FString, FString> OCIOContext;
+};
\ No newline at end of file
diff --git a/Engine/Plugins/Media/AppleProResMedia/Source/AppleProResMedia/Public/AppleProResEncoder/AppleProResEncoder.h b/Engine/Plugins/Media/AppleProResMedia/Source/AppleProResMedia/Public/AppleProResEncoder/AppleProResEncoder.h
--- a/Engine/Plugins/Media/AppleProResMedia/Source/AppleProResMedia/Public/AppleProResEncoder/AppleProResEncoder.h
+++ b/Engine/Plugins/Media/AppleProResMedia/Source/AppleProResMedia/Public/AppleProResEncoder/AppleProResEncoder.h
@@ -63,9 +63,12 @@ struct FAppleProResEncoderOptions
 		, FrameRate(30, 1)
 		, Codec(EAppleProResEncoderCodec::ProRes_422)
 		, ColorPrimaries(EAppleProResEncoderColorPrimaries::CD_HDREC709)
+		, bConvertToSrgb(true)
 		, ScanMode(EAppleProResEncoderScanMode::IM_PROGRESSIVE_SCAN)
+		, MaxNumberOfEncodingThreads(0)
 		, bWriteAlpha(false)
 		, bDropFrameTimecode(false)
+		, FrameNumberOffset(0)
 	{}
 
 	/** The absolute path on disk to try and save the video file to. */
@@ -86,6 +89,9 @@ struct FAppleProResEncoderOptions
 	/** Which color primaries do we use? Only Rec 709 is well tested right now. */
 	EAppleProResEncoderColorPrimaries ColorPrimaries;
 
+	/** Whether the data should be converted to sRGB before being sent to the encoder. Should not be used when OCIO is active. */
+	bool bConvertToSrgb;
+
 	/** Which scan mode do we use? Only Progressive is tested right now. */
 	EAppleProResEncoderScanMode ScanMode;
 	
@@ -122,9 +128,13 @@ public:
 
 	/** Appends a new frame onto the output file. */
 	bool WriteFrame(const FImagePixelData* InPixelData);
+	
 	/** Appends a new audio sample onto the audio stream. */
 	bool WriteAudioSample(const TArrayView<int16>& InAudioSamples);
 
+	/** Gets the options that the encoder was initialized with. */
+	const FAppleProResEncoderOptions& GetOptions() const;
+
 private:
 	bool InitializeVideoTrack();
 	bool InitializeTimecodeTrack();
diff --git a/Engine/Plugins/Media/AvidDNxMedia/Source/Source/Private/AvidDNxEncoder/AvidDNxEncoder.cpp b/Engine/Plugins/Media/AvidDNxMedia/Source/Source/Private/AvidDNxEncoder/AvidDNxEncoder.cpp
--- a/Engine/Plugins/Media/AvidDNxMedia/Source/Source/Private/AvidDNxEncoder/AvidDNxEncoder.cpp
+++ b/Engine/Plugins/Media/AvidDNxMedia/Source/Source/Private/AvidDNxEncoder/AvidDNxEncoder.cpp
@@ -484,6 +484,11 @@ bool FAvidDNxEncoder::WriteFrame_16bit(const FFloat16Color* InFrameData)
 	return WriteFrame_16bit_Impl<FY0CbY1Cr_16bit>(InFrameData);
 }
 
+const FAvidDNxEncoderOptions& FAvidDNxEncoder::GetOptions() const
+{
+	return Options;
+}
+
 template <typename EncodedBufferType>
 bool FAvidDNxEncoder::WriteFrame_16bit_Impl(const FFloat16Color* InFrameData)
 {
diff --git a/Engine/Plugins/Media/AvidDNxMedia/Source/Source/Private/MovieGraphAvidDNxHRNode.cpp b/Engine/Plugins/Media/AvidDNxMedia/Source/Source/Private/MovieGraphAvidDNxHRNode.cpp
--- a/Engine/Plugins/Media/AvidDNxMedia/Source/Source/Private/MovieGraphAvidDNxHRNode.cpp
+++ b/Engine/Plugins/Media/AvidDNxMedia/Source/Source/Private/MovieGraphAvidDNxHRNode.cpp
@@ -73,20 +73,19 @@ TUniquePtr<MovieRenderGraph::IVideoCodecWriter> UMovieGraphAvidDNxHRNode::Initia
 	Options.FrameRate = EffectiveFrameRate;
 	Options.bCompress = true;
 	Options.NumberOfEncodingThreads = 4;
-	
-	TUniquePtr<FAvidWriter> NewWriter = MakeUnique<FAvidWriter>();
-	NewWriter->Writer = MakeUnique<FAvidDNxEncoder>(Options);
-	NewWriter->FileName = InFileName;
 
 	// If OCIO is enabled, don't do additional color conversion. RGB444 12-bit is never converted to sRGB.
 	if (EvaluatedNode->Quality == EAvidDNxEncoderQuality::RGB444_12bit)
 	{
-		NewWriter->bConvertToSrgb = false;
+		Options.bConvertToSrgb = false;
 	}
 	else
 	{
-		NewWriter->bConvertToSrgb = !(bOverride_OCIOConfiguration && OCIOConfiguration.bIsEnabled && bAllowOCIO);
+		Options.bConvertToSrgb = !(bOverride_OCIOConfiguration && OCIOConfiguration.bIsEnabled && bAllowOCIO);
 	}
+	
+	TUniquePtr<FAvidWriter> NewWriter = MakeUnique<FAvidWriter>();
+	NewWriter->Writer = MakeUnique<FAvidDNxEncoder>(Options);
 
 	CachedPipeline = InPipeline;
 	
@@ -105,7 +104,7 @@ bool UMovieGraphAvidDNxHRNode::Initialize_EncodeThread(MovieRenderGraph::IVideoC
 	return true;
 }
 
-void UMovieGraphAvidDNxHRNode::WriteFrame_EncodeThread(MovieRenderGraph::IVideoCodecWriter* InWriter, FImagePixelData* InPixelData, TArray<FMovieGraphPassData>&& InCompositePasses, TObjectPtr<UMovieGraphEvaluatedConfig> InEvaluatedConfig)
+void UMovieGraphAvidDNxHRNode::WriteFrame_EncodeThread(MovieRenderGraph::IVideoCodecWriter* InWriter, FImagePixelData* InPixelData, TArray<FMovieGraphPassData>&& InCompositePasses, TObjectPtr<UMovieGraphEvaluatedConfig> InEvaluatedConfig, const FString& InBranchName)
 {
 	const FAvidWriter* CodecWriter = static_cast<FAvidWriter*>(InWriter);
 	
@@ -120,7 +119,7 @@ void UMovieGraphAvidDNxHRNode::WriteFrame_EncodeThread(MovieRenderGraph::IVideoC
 
 	// Quantize our 16-bit float data to 8/16-bit and apply sRGB if needed
 	const int32 BitDepth = ((EvaluatedNode->Quality == EAvidDNxEncoderQuality::HQX_10bit) || (EvaluatedNode->Quality == EAvidDNxEncoderQuality::RGB444_12bit)) ? 16 : 8;
-	TUniquePtr<FImagePixelData> QuantizedPixelData = UE::MoviePipeline::QuantizeImagePixelDataToBitDepth(InPixelData, BitDepth, nullptr, InWriter->bConvertToSrgb);
+	TUniquePtr<FImagePixelData> QuantizedPixelData = UE::MoviePipeline::QuantizeImagePixelDataToBitDepth(InPixelData, BitDepth, nullptr, CodecWriter->Writer->GetOptions().bConvertToSrgb);
 
 	TArray<FPixelPreProcessor> PixelPreProcessors;
 
diff --git a/Engine/Plugins/Media/AvidDNxMedia/Source/Source/Private/MovieGraphAvidDNxHRNode.h b/Engine/Plugins/Media/AvidDNxMedia/Source/Source/Private/MovieGraphAvidDNxHRNode.h
--- a/Engine/Plugins/Media/AvidDNxMedia/Source/Source/Private/MovieGraphAvidDNxHRNode.h
+++ b/Engine/Plugins/Media/AvidDNxMedia/Source/Source/Private/MovieGraphAvidDNxHRNode.h
@@ -41,7 +41,7 @@ protected:
 	// UMovieGraphVideoOutputNode Interface
 	virtual TUniquePtr<MovieRenderGraph::IVideoCodecWriter> Initialize_GameThread(UMovieGraphPipeline* InPipeline, TObjectPtr<UMovieGraphEvaluatedConfig> InEvaluatedConfig, const FString& InBranchName, const FString& InFileName, FIntPoint InResolution, EImagePixelType InPixelType, ERGBFormat InPixelFormat, uint8 InBitDepth, uint8 InNumChannels, bool bAllowOCIO) override;
 	virtual bool Initialize_EncodeThread(MovieRenderGraph::IVideoCodecWriter* InWriter) override;
-	virtual void WriteFrame_EncodeThread(MovieRenderGraph::IVideoCodecWriter* InWriter, FImagePixelData* InPixelData, TArray<FMovieGraphPassData>&& InCompositePasses, TObjectPtr<UMovieGraphEvaluatedConfig> InEvaluatedConfig) override;
+	virtual void WriteFrame_EncodeThread(MovieRenderGraph::IVideoCodecWriter* InWriter, FImagePixelData* InPixelData, TArray<FMovieGraphPassData>&& InCompositePasses, TObjectPtr<UMovieGraphEvaluatedConfig> InEvaluatedConfig, const FString& InBranchName) override;
 	virtual void BeginFinalize_EncodeThread(MovieRenderGraph::IVideoCodecWriter* InWriter) override;
 	virtual void Finalize_EncodeThread(MovieRenderGraph::IVideoCodecWriter* InWriter) override;
 	virtual const TCHAR* GetFilenameExtension() const override;
diff --git a/Engine/Plugins/Media/AvidDNxMedia/Source/Source/Public/AvidDNxEncoder/AvidDNxEncoder.h b/Engine/Plugins/Media/AvidDNxMedia/Source/Source/Public/AvidDNxEncoder/AvidDNxEncoder.h
--- a/Engine/Plugins/Media/AvidDNxMedia/Source/Source/Public/AvidDNxEncoder/AvidDNxEncoder.h
+++ b/Engine/Plugins/Media/AvidDNxMedia/Source/Source/Public/AvidDNxEncoder/AvidDNxEncoder.h
@@ -63,6 +63,9 @@ struct FAvidDNxEncoderOptions
 
 	/** The quality setting that the encoder should use. */
 	EAvidDNxEncoderQuality Quality;
+	
+	/** Whether the data should be converted to sRGB before being sent to the encoder. Should not be used when OCIO is active. */
+	bool bConvertToSrgb;
 
 	/** Frame Rate of the output video. */
 	FFrameRate FrameRate;
@@ -96,6 +99,9 @@ public:
 	/** Appends a new frame onto the output file (16-bit). */
 	bool WriteFrame_16bit(const FFloat16Color* InFrameData);
 
+	/** Gets the options that the encoder was initialized with. */
+	const FAvidDNxEncoderOptions& GetOptions() const;
+
 private:
 	bool InitializeCompressedEncoder();
 	bool InitializeUncompressedEncoder();
diff --git a/Engine/Plugins/MovieScene/MovieRenderPipeline/Source/MovieRenderPipelineCore/Private/Graph/Nodes/MovieGraphVideoOutputNode.cpp b/Engine/Plugins/MovieScene/MovieRenderPipeline/Source/MovieRenderPipelineCore/Private/Graph/Nodes/MovieGraphVideoOutputNode.cpp
--- a/Engine/Plugins/MovieScene/MovieRenderPipeline/Source/MovieRenderPipelineCore/Private/Graph/Nodes/MovieGraphVideoOutputNode.cpp
+++ b/Engine/Plugins/MovieScene/MovieRenderPipeline/Source/MovieRenderPipelineCore/Private/Graph/Nodes/MovieGraphVideoOutputNode.cpp
@@ -104,12 +104,12 @@ void UMovieGraphVideoOutputNode::OnReceiveImageDataImpl(UMovieGraphPipeline* InP
 			}
 
 			// Notify the encoder to write this frame
-			this->WriteFrame_EncodeThread(OutputWriter->CodecWriter.Get(), RawRenderPassData, MoveTemp(CompositesForThisCamera), EvaluatedConfig);
+			this->WriteFrame_EncodeThread(OutputWriter->CodecWriter.Get(), RawRenderPassData, MoveTemp(CompositesForThisCamera), EvaluatedConfig, RenderPassData.Key.RootBranchName.ToString());
 		}
 		else
 		{
 			TArray<FMovieGraphPassData> Dummy;
-			this->WriteFrame_EncodeThread(OutputWriter->CodecWriter.Get(), RawRenderPassData, MoveTemp(Dummy), EvaluatedConfig);
+			this->WriteFrame_EncodeThread(OutputWriter->CodecWriter.Get(), RawRenderPassData, MoveTemp(Dummy), EvaluatedConfig, RenderPassData.Key.RootBranchName.ToString());
 		}
 	}
 }
@@ -165,6 +165,9 @@ void UMovieGraphVideoOutputNode::GetOutputFilePaths(const UMovieGraphPipeline* I
 	UMovieGraphGlobalOutputSettingNode* OutputSettingsNode =
 		EvaluatedConfig->GetSettingForBranch<UMovieGraphGlobalOutputSettingNode>(GlobalsPinName, bIncludeCDOs, bExactMatch);
 
+	UMovieGraphVideoOutputNode* EvaluatedNode = Cast<UMovieGraphVideoOutputNode>(
+		EvaluatedConfig->GetSettingForBranch(GetClass(), InRenderPassData.Key.RootBranchName, bIncludeCDOs, bExactMatch));
+
 	const TMap<FString, FString> FileNameFormatOverrides =  {
 		{TEXT("camera_name"), InRenderPassData.Key.CameraName},
 		{TEXT("render_pass"), InRenderPassData.Key.RendererName},
@@ -172,11 +175,11 @@ void UMovieGraphVideoOutputNode::GetOutputFilePaths(const UMovieGraphPipeline* I
 	};
 	FMovieGraphFilenameResolveParams ResolveParams =
 		FMovieGraphFilenameResolveParams::MakeResolveParams(InRenderPassData.Key, InPipeline, EvaluatedConfig, InRawFrameData->TraversalContext, FileNameFormatOverrides);
-	ResolveParams.FileNameOverride = FileNameFormat;
+	ResolveParams.FileNameOverride = EvaluatedNode->FileNameFormat;
 	ResolveParams.Version = Shot ? Shot->ShotInfo.VersionNumber : UMovieGraphBlueprintLibrary::ResolveVersionNumber(ResolveParams);
 	
 	FMovieGraphResolveArgs FinalFormatArgs;
-	FString FileNameFormatString = FileNameFormat;
+	FString FileNameFormatString = EvaluatedNode->FileNameFormat;
 	
 	// If we're writing more than one render pass out, we need to ensure the file name has the format string in it so we don't
 	// overwrite the same file multiple times. Burn In overlays don't count because they get composited on top of an existing file.
diff --git a/Engine/Plugins/MovieScene/MovieRenderPipeline/Source/MovieRenderPipelineCore/Public/Graph/Nodes/MovieGraphVideoOutputNode.h b/Engine/Plugins/MovieScene/MovieRenderPipeline/Source/MovieRenderPipelineCore/Public/Graph/Nodes/MovieGraphVideoOutputNode.h
--- a/Engine/Plugins/MovieScene/MovieRenderPipeline/Source/MovieRenderPipelineCore/Public/Graph/Nodes/MovieGraphVideoOutputNode.h
+++ b/Engine/Plugins/MovieScene/MovieRenderPipeline/Source/MovieRenderPipelineCore/Public/Graph/Nodes/MovieGraphVideoOutputNode.h
@@ -14,18 +14,12 @@ namespace MovieRenderGraph
 {
 	struct IVideoCodecWriter
 	{
-		/** The filename to actually write to disk with. This may include "de-duplication" numbers (MyFile (1).ext etc.) */
-		FString FileName;
-		
 		/**
 		 * The filename without de-duplication numbers, used to match up multiple incoming frames back to the same writer.
 		 * We use this when looking for existing writers so that we can avoid the de-duplication numbers perpetually increasing
 		 * due to the file existing on disk after the first frame comes in, and then the next one de-duplicating to one more than that.
 		 */
 		FString StableFileName;
-
-		/** Whether the data should be converted to sRGB. Should not be used when OCIO is active. */
-		bool bConvertToSrgb;
 	};
 }
 
@@ -51,7 +45,7 @@ protected:
 
 	virtual TUniquePtr<MovieRenderGraph::IVideoCodecWriter> Initialize_GameThread(UMovieGraphPipeline* InPipeline, TObjectPtr<UMovieGraphEvaluatedConfig> InEvaluatedConfig, const FString& InBranchName, const FString& InFileName, FIntPoint InResolution, EImagePixelType InPixelType, ERGBFormat InPixelFormat, uint8 InBitDepth, uint8 InNumChannels, bool bAllowOCIO)  PURE_VIRTUAL(UMovieGraphVideoOutputNode::Initialize_GameThread, return nullptr; );
 	virtual bool Initialize_EncodeThread(MovieRenderGraph::IVideoCodecWriter* InWriter) PURE_VIRTUAL(UMovieGraphVideoOutputNode::Initialize_EncodeThread, return true;);
-	virtual void WriteFrame_EncodeThread(MovieRenderGraph::IVideoCodecWriter* InWriter, FImagePixelData* InPixelData, TArray<FMovieGraphPassData>&& InCompositePasses, TObjectPtr<UMovieGraphEvaluatedConfig> InEvaluatedConfig) PURE_VIRTUAL(UMovieGraphVideoOutputNode::WriteFrame_EncodeThread);
+	virtual void WriteFrame_EncodeThread(MovieRenderGraph::IVideoCodecWriter* InWriter, FImagePixelData* InPixelData, TArray<FMovieGraphPassData>&& InCompositePasses, TObjectPtr<UMovieGraphEvaluatedConfig> InEvaluatedConfig, const FString& InBranchName) PURE_VIRTUAL(UMovieGraphVideoOutputNode::WriteFrame_EncodeThread);
 	virtual void BeginFinalize_EncodeThread(MovieRenderGraph::IVideoCodecWriter* InWriter) PURE_VIRTUAL(UMovieGraphVideoOutputNode::BeginFinalize_EncodeThread);
 	virtual void Finalize_EncodeThread(MovieRenderGraph::IVideoCodecWriter* InWriter) PURE_VIRTUAL(UMovieGraphVideoOutputNode::Finalize_EncodeThread);
 	virtual const TCHAR* GetFilenameExtension() const PURE_VIRTUAL(UMovieGraphVideoOutputNode::GetFilenameExtension, return TEXT(""););
