MoviePipeline: Updated the Global Output Settings node to include timecode properties, rather than including it on several individual nodes.

This also has the effect of ensuring that the burn-in timecode matches the timecode that's embedded in Avid DNxHR and ProRes movies, because the timecode properties are now impacting the timecode data that is supplied to the burn-in.

#rnx

diff --git a/Engine/Plugins/Media/AppleProResMedia/Source/AppleProResMedia/Private/MovieGraphAppleProResNode.cpp b/Engine/Plugins/Media/AppleProResMedia/Source/AppleProResMedia/Private/MovieGraphAppleProResNode.cpp
--- a/Engine/Plugins/Media/AppleProResMedia/Source/AppleProResMedia/Private/MovieGraphAppleProResNode.cpp
+++ b/Engine/Plugins/Media/AppleProResMedia/Source/AppleProResMedia/Private/MovieGraphAppleProResNode.cpp
@@ -15,7 +15,6 @@
 
 UMovieGraphAppleProResNode::UMovieGraphAppleProResNode()
 	: Quality(EAppleProResEncoderCodec::ProRes_422LT)
-	, bDropFrameTimecode(true)	// Defaults to true because most 29.97 FPS content uses this
 {
 	
 }
@@ -77,7 +76,7 @@ TUniquePtr<MovieRenderGraph::IVideoCodecWriter> UMovieGraphAppleProResNode::Init
 	Options.ColorPrimaries = EAppleProResEncoderColorPrimaries::CD_HDREC709; // Force Rec 709 for now
 	Options.ScanMode = EAppleProResEncoderScanMode::IM_PROGRESSIVE_SCAN; // No interlace sources
 	Options.bWriteAlpha = true;
-	Options.bDropFrameTimecode = bDropFrameTimecode;
+	Options.bDropFrameTimecode = OutputSetting->bDropFrameTimecode;
 
 	// If OCIO is enabled, don't do additional color conversion
 	Options.bConvertToSrgb = !(EvaluatedNode->bOverride_OCIOConfiguration && EvaluatedNode->OCIOConfiguration.bIsEnabled && InInitializationContext.bAllowOCIO);
@@ -104,11 +103,15 @@ bool UMovieGraphAppleProResNode::Initialize_EncodeThread(MovieRenderGraph::IVide
 
 void UMovieGraphAppleProResNode::WriteFrame_EncodeThread(MovieRenderGraph::IVideoCodecWriter* InWriter, FImagePixelData* InPixelData, TArray<FMovieGraphPassData>&& InCompositePasses, TObjectPtr<UMovieGraphEvaluatedConfig> InEvaluatedConfig, const FString& InBranchName)
 {
-	constexpr bool bIncludeCDOs = false;
+	bool bIncludeCDOs = false;
 	constexpr bool bExactMatch = true;
 	const UMovieGraphAppleProResNode* EvaluatedNode = Cast<UMovieGraphAppleProResNode>(
 		InEvaluatedConfig->GetSettingForBranch(GetClass(), FName(InBranchName), bIncludeCDOs, bExactMatch));
 	checkf(EvaluatedNode, TEXT("Apple ProRes node could not be found in the graph in branch [%s]."), *InBranchName);
+
+	bIncludeCDOs = true;
+	const UMovieGraphGlobalOutputSettingNode* OutputSettingNode =
+		InEvaluatedConfig->GetSettingForBranch<UMovieGraphGlobalOutputSettingNode>(GlobalsPinName, bIncludeCDOs, bExactMatch);
 	
 	const FProResWriter* CodecWriter = static_cast<FProResWriter*>(InWriter);
 	
@@ -117,9 +120,9 @@ void UMovieGraphAppleProResNode::WriteFrame_EncodeThread(MovieRenderGraph::IVide
 	// Translate our Movie Pipeline specific payload to a ProRes Encoder specific payload.
 	const TSharedPtr<FAppleProResEncoder::FTimecodePayload, ESPMode::ThreadSafe> ProResPayload = MakeShared<FAppleProResEncoder::FTimecodePayload, ESPMode::ThreadSafe>();
 
-	if (EvaluatedNode->bOverride_CustomTimecodeStart)
+	if (OutputSettingNode->bOverride_CustomTimecodeStart)
 	{
-		const int32 CustomTimecodeStartFrame = EvaluatedNode->CustomTimecodeStart.ToFrameNumber(CodecWriter->Writer->GetOptions().FrameRate).Value;
+		const int32 CustomTimecodeStartFrame = OutputSettingNode->CustomTimecodeStart.ToFrameNumber(CodecWriter->Writer->GetOptions().FrameRate).Value;
 		const FMovieGraphTimeStepData& TimeData = GraphPayload->TraversalContext.Time;
 		
 		// When using a custom timecode start, just use the root-level frame number (relative to zero) offset by the custom timecode start
diff --git a/Engine/Plugins/Media/AppleProResMedia/Source/AppleProResMedia/Private/MovieGraphAppleProResNode.h b/Engine/Plugins/Media/AppleProResMedia/Source/AppleProResMedia/Private/MovieGraphAppleProResNode.h
--- a/Engine/Plugins/Media/AppleProResMedia/Source/AppleProResMedia/Private/MovieGraphAppleProResNode.h
+++ b/Engine/Plugins/Media/AppleProResMedia/Source/AppleProResMedia/Private/MovieGraphAppleProResNode.h
@@ -49,12 +49,6 @@ public:
 	UPROPERTY(EditAnywhere, BlueprintReadWrite, Category = Overrides, meta = (InlineEditConditionToggle))
 	uint8 bOverride_Quality : 1;
 
-	UPROPERTY(EditAnywhere, BlueprintReadWrite, Category = Overrides, meta = (InlineEditConditionToggle))
-	uint8 bOverride_CustomTimecodeStart : 1;
-
-	UPROPERTY(EditAnywhere, BlueprintReadWrite, Category = Overrides, meta = (InlineEditConditionToggle))
-	uint8 bOverride_bDropFrameTimecode : 1;
-
 	UPROPERTY(EditAnywhere, BlueprintReadWrite, Category = Overrides, meta = (InlineEditConditionToggle))
 	uint8 bOverride_OCIOConfiguration : 1;
 
@@ -65,14 +59,6 @@ public:
 	UPROPERTY(EditAnywhere, BlueprintReadWrite, Category = "Apple ProRes", meta = (EditCondition = "bOverride_Quality"))
 	EAppleProResEncoderCodec Quality;
 
-	/** Start the timecode at a specific value, rather than the value coming from the Level Sequence. */
-	UPROPERTY(EditAnywhere, BlueprintReadWrite, Category = "Apple ProRes", meta = (EditCondition = "bOverride_CustomTimecodeStart"))
-	FTimecode CustomTimecodeStart;
-
-	/** Whether the embedded timecode track should be written using drop-frame format. Only applicable if the sequence framerate is 29.97 */
-	UPROPERTY(EditAnywhere, BlueprintReadWrite, Category = "Apple ProRes", DisplayName = "Use DF Timecode if 29.97 FPS", meta = (EditCondition = "bOverride_bDropFrameTimecode"))
-	bool bDropFrameTimecode;
-
 	/**
 	* OCIO configuration/transform settings.
 	*
diff --git a/Engine/Plugins/Media/AvidDNxMedia/Source/Source/Private/MovieGraphAvidDNxHRNode.cpp b/Engine/Plugins/Media/AvidDNxMedia/Source/Source/Private/MovieGraphAvidDNxHRNode.cpp
--- a/Engine/Plugins/Media/AvidDNxMedia/Source/Source/Private/MovieGraphAvidDNxHRNode.cpp
+++ b/Engine/Plugins/Media/AvidDNxMedia/Source/Source/Private/MovieGraphAvidDNxHRNode.cpp
@@ -14,7 +14,6 @@
 #include "Styling/AppStyle.h"
 
 UMovieGraphAvidDNxHRNode::UMovieGraphAvidDNxHRNode()
-	: bDropFrameTimecode(true)
 {
 }
 
@@ -67,15 +66,15 @@ TUniquePtr<MovieRenderGraph::IVideoCodecWriter> UMovieGraphAvidDNxHRNode::Initia
 
 	// Determine the timecode that the movie should be started at
 	FTimecode StartTimecode;
-	if (EvaluatedNode->bOverride_CustomTimecodeStart)
+	if (OutputSetting->bOverride_CustomTimecodeStart)
 	{
 		const int32 OutputFrameNumber = InInitializationContext.TraversalContext->Time.OutputFrameNumber;
 
 		// When using a custom timecode start, just use the root-level frame number (relative to zero) offset by the custom timecode start
 		StartTimecode = FTimecode::FromFrameNumber(
-			OutputFrameNumber + EvaluatedNode->CustomTimecodeStart.ToFrameNumber(EffectiveFrameRate).Value,
+			OutputFrameNumber + OutputSetting->CustomTimecodeStart.ToFrameNumber(EffectiveFrameRate).Value,
 			EffectiveFrameRate,
-			EvaluatedNode->bDropFrameTimecode);
+			OutputSetting->bDropFrameTimecode);
 	}
 	else
 	{
@@ -91,7 +90,7 @@ TUniquePtr<MovieRenderGraph::IVideoCodecWriter> UMovieGraphAvidDNxHRNode::Initia
 	Options.FrameRate = EffectiveFrameRate;
 	Options.bCompress = true;
 	Options.NumberOfEncodingThreads = 4;
-	Options.bDropFrameTimecode = EvaluatedNode->bDropFrameTimecode;
+	Options.bDropFrameTimecode = OutputSetting->bDropFrameTimecode;
 	Options.StartTimecode = StartTimecode;
 
 	// If OCIO is enabled, don't do additional color conversion. RGB444 12-bit is never converted to sRGB.
diff --git a/Engine/Plugins/Media/AvidDNxMedia/Source/Source/Private/MovieGraphAvidDNxHRNode.h b/Engine/Plugins/Media/AvidDNxMedia/Source/Source/Private/MovieGraphAvidDNxHRNode.h
--- a/Engine/Plugins/Media/AvidDNxMedia/Source/Source/Private/MovieGraphAvidDNxHRNode.h
+++ b/Engine/Plugins/Media/AvidDNxMedia/Source/Source/Private/MovieGraphAvidDNxHRNode.h
@@ -50,12 +50,6 @@ protected:
 public:
 	UPROPERTY(EditAnywhere, BlueprintReadWrite, Category = Overrides, meta = (InlineEditConditionToggle))
 	uint8 bOverride_Quality : 1;
-	
-	UPROPERTY(EditAnywhere, BlueprintReadWrite, Category = Overrides, meta = (InlineEditConditionToggle))
-	uint8 bOverride_CustomTimecodeStart : 1;
-
-	UPROPERTY(EditAnywhere, BlueprintReadWrite, Category = Overrides, meta = (InlineEditConditionToggle))
-	uint8 bOverride_bDropFrameTimecode : 1;
 
 	UPROPERTY(EditAnywhere, BlueprintReadWrite, Category = Overrides, meta = (InlineEditConditionToggle))
 	uint8 bOverride_OCIOConfiguration : 1;
@@ -67,14 +61,6 @@ public:
 	UPROPERTY(EditAnywhere, BlueprintReadWrite, Category="Avid DNxHR", meta=(EditCondition="bOverride_Quality"))
 	EAvidDNxEncoderQuality Quality = EAvidDNxEncoderQuality::HQ_8bit;
 
-	/** Start the timecode at a specific value, rather than the value coming from the Level Sequence. */
-	UPROPERTY(EditAnywhere, BlueprintReadWrite, Category = "Avid DNxHR", meta = (EditCondition = "bOverride_CustomTimecodeStart"))
-	FTimecode CustomTimecodeStart;
-
-	/** Whether the embedded timecode track should be written using drop-frame format. Only applicable if the sequence framerate is 29.97. */
-	UPROPERTY(EditAnywhere, BlueprintReadWrite, Category = "Avid DNxHR", DisplayName = "Use DF Timecode if 29.97 FPS", meta = (EditCondition = "bOverride_bDropFrameTimecode"))
-	bool bDropFrameTimecode;
-
 	/**
 	* OCIO configuration/transform settings.
 	*
diff --git a/Engine/Plugins/MovieScene/MovieRenderPipeline/Source/MovieRenderPipelineCore/Private/Graph/MovieGraphCoreTimeStep.cpp b/Engine/Plugins/MovieScene/MovieRenderPipeline/Source/MovieRenderPipelineCore/Private/Graph/MovieGraphCoreTimeStep.cpp
--- a/Engine/Plugins/MovieScene/MovieRenderPipeline/Source/MovieRenderPipelineCore/Private/Graph/MovieGraphCoreTimeStep.cpp
+++ b/Engine/Plugins/MovieScene/MovieRenderPipeline/Source/MovieRenderPipelineCore/Private/Graph/MovieGraphCoreTimeStep.cpp
@@ -438,18 +438,29 @@ void UMovieGraphCoreTimeStep::TickProducingFrames()
 		// the wrong one. Because temporal sub-sampling isn't centered around a frame (the centering is done via the final eval time) we can just subtract TSI*TPS to get our centered value.
 		const FFrameTime CenteringOffset = CurrentFrameData.TemporalSampleIndex * CurrentFrameMetrics.FrameTimePerTemporalSample;
 		FFrameTime CenteredFrameTime = CurrentCameraCut->ShotInfo.CurrentTimeInRoot - CenteringOffset;
+		FFrameTime ShotCenteredFrameTime = CenteredFrameTime * CurrentCameraCut->ShotInfo.OuterToInnerTransform;
 	
 		const FFrameRate SourceFrameRate = GetOwningGraph()->GetDataSourceInstance()->GetDisplayRate();
 		const FFrameRate EffectiveFrameRate = UMovieGraphBlueprintLibrary::GetEffectiveFrameRate(OutputSetting, SourceFrameRate);
 		const FFrameRate TickResolution = GetOwningGraph()->GetDataSourceInstance()->GetTickResolution();
 
-		constexpr bool bDropFrame = false;
-		CurrentTimeStepData.RootFrameNumber = FFrameRate::TransformTime(CenteredFrameTime, TickResolution, EffectiveFrameRate).RoundToFrame();
+		// Enable DF timecodes if the framerate is 29.97 and the user requested DF timecodes
+		const FFrameRate TwentyNineNineSeven = FFrameRate(30000, 1001);
+		const bool bDropFrame = OutputSetting->bDropFrameTimecode && (EffectiveFrameRate == TwentyNineNineSeven);
+
+		// If using a custom timecode start, the frame number for the root and shot need to be changed
+		const FFrameNumber RootFrameNumber = OutputSetting->bOverride_CustomTimecodeStart
+			? CurrentTimeStepData.OutputFrameNumber + OutputSetting->CustomTimecodeStart.ToFrameNumber(EffectiveFrameRate).Value
+			: FFrameRate::TransformTime(CenteredFrameTime, TickResolution, EffectiveFrameRate).RoundToFrame();
+		const FFrameNumber ShotFrameNumber = OutputSetting->bOverride_CustomTimecodeStart
+			? CurrentTimeStepData.OutputFrameNumber + OutputSetting->CustomTimecodeStart.ToFrameNumber(EffectiveFrameRate).Value
+			: FFrameRate::TransformTime(ShotCenteredFrameTime, TickResolution, EffectiveFrameRate).RoundToFrame();
+
+		CurrentTimeStepData.RootFrameNumber = RootFrameNumber;
 		CurrentTimeStepData.RootTimeCode = FTimecode::FromFrameNumber(CurrentTimeStepData.RootFrameNumber, EffectiveFrameRate, bDropFrame);
 
 		// Calculate metrics for the shot as well
-		CenteredFrameTime = CenteredFrameTime * CurrentCameraCut->ShotInfo.OuterToInnerTransform;
-		CurrentTimeStepData.ShotFrameNumber = FFrameRate::TransformTime(CenteredFrameTime, TickResolution, EffectiveFrameRate).RoundToFrame();
+		CurrentTimeStepData.ShotFrameNumber = ShotFrameNumber;
 		CurrentTimeStepData.ShotTimeCode = FTimecode::FromFrameNumber(CurrentTimeStepData.ShotFrameNumber, EffectiveFrameRate, bDropFrame);
 	}
 
diff --git a/Engine/Plugins/MovieScene/MovieRenderPipeline/Source/MovieRenderPipelineCore/Private/Graph/Nodes/MovieGraphGlobalOutputSettingNode.cpp b/Engine/Plugins/MovieScene/MovieRenderPipeline/Source/MovieRenderPipelineCore/Private/Graph/Nodes/MovieGraphGlobalOutputSettingNode.cpp
--- a/Engine/Plugins/MovieScene/MovieRenderPipeline/Source/MovieRenderPipelineCore/Private/Graph/Nodes/MovieGraphGlobalOutputSettingNode.cpp
+++ b/Engine/Plugins/MovieScene/MovieRenderPipeline/Source/MovieRenderPipelineCore/Private/Graph/Nodes/MovieGraphGlobalOutputSettingNode.cpp
@@ -16,6 +16,7 @@ UMovieGraphGlobalOutputSettingNode::UMovieGraphGlobalOutputSettingNode()
 	, HandleFrameCount(0)
 	, CustomPlaybackRangeStartFrame(0)
 	, CustomPlaybackRangeEndFrame(0)
+	, bDropFrameTimecode(true)	// Defaults to true because most 29.97 FPS content uses this
 	, bFlushDiskWritesPerShot(false)
 {
 	OutputDirectory.Path = TEXT("{project_dir}/Saved/MovieRenders/");
diff --git a/Engine/Plugins/MovieScene/MovieRenderPipeline/Source/MovieRenderPipelineCore/Public/Graph/MovieGraphTimeStepData.h b/Engine/Plugins/MovieScene/MovieRenderPipeline/Source/MovieRenderPipelineCore/Public/Graph/MovieGraphTimeStepData.h
--- a/Engine/Plugins/MovieScene/MovieRenderPipeline/Source/MovieRenderPipelineCore/Public/Graph/MovieGraphTimeStepData.h
+++ b/Engine/Plugins/MovieScene/MovieRenderPipeline/Source/MovieRenderPipelineCore/Public/Graph/MovieGraphTimeStepData.h
@@ -129,7 +129,7 @@ public:
 	UPROPERTY(BlueprintReadWrite, EditAnywhere, Category = "Movie Graph")
 	TObjectPtr<UMovieGraphEvaluatedConfig> EvaluatedConfig;
 
-	/** The current timecode at the root (sequence) level. */
+	/** The current timecode at the root (sequence) level. Note that this is adjusted according to the timecode settings on the Global Output Settings node. */
 	UPROPERTY(BlueprintReadWrite, EditAnywhere, Category = "Movie Graph")
 	FTimecode RootTimeCode;
 
@@ -140,7 +140,7 @@ public:
 	UPROPERTY(BlueprintReadWrite, EditAnywhere, Category = "Movie Graph")
 	FFrameNumber RootFrameNumber;
 
-	/** The current timecode at the shot level. */
+	/** The current timecode at the shot level. Note that this is adjusted according to the timecode settings on the Global Output Settings node. */
 	UPROPERTY(BlueprintReadWrite, EditAnywhere, Category = "Movie Graph")
 	FTimecode ShotTimeCode;
 
diff --git a/Engine/Plugins/MovieScene/MovieRenderPipeline/Source/MovieRenderPipelineCore/Public/Graph/Nodes/MovieGraphGlobalOutputSettingNode.h b/Engine/Plugins/MovieScene/MovieRenderPipeline/Source/MovieRenderPipelineCore/Public/Graph/Nodes/MovieGraphGlobalOutputSettingNode.h
--- a/Engine/Plugins/MovieScene/MovieRenderPipeline/Source/MovieRenderPipelineCore/Public/Graph/Nodes/MovieGraphGlobalOutputSettingNode.h
+++ b/Engine/Plugins/MovieScene/MovieRenderPipeline/Source/MovieRenderPipelineCore/Public/Graph/Nodes/MovieGraphGlobalOutputSettingNode.h
@@ -6,6 +6,7 @@
 #include "Graph/MovieGraphNode.h"
 
 #include "Misc/FrameRate.h"
+#include "Misc/Timecode.h"
 
 #include "MovieGraphGlobalOutputSettingNode.generated.h"
 
@@ -77,6 +78,12 @@ public:
 	UPROPERTY(EditAnywhere, BlueprintReadWrite, Category = Overrides, meta = (InlineEditConditionToggle))
 	uint8 bOverride_CustomPlaybackRangeEndFrame : 1;
 
+	UPROPERTY(EditAnywhere, BlueprintReadWrite, Category = Overrides, meta = (InlineEditConditionToggle))
+	uint8 bOverride_CustomTimecodeStart : 1;
+
+	UPROPERTY(EditAnywhere, BlueprintReadWrite, Category = Overrides, meta = (InlineEditConditionToggle))
+	uint8 bOverride_bDropFrameTimecode : 1;
+
 	UPROPERTY(EditAnywhere, BlueprintReadWrite, Category = Overrides, meta = (InlineEditConditionToggle))
 	uint8 bOverride_VersioningSettings : 1;
 
@@ -125,6 +132,14 @@ public:
 	UPROPERTY(EditAnywhere, BlueprintReadWrite, Category = "Frames", meta = (EditCondition = "bOverride_CustomPlaybackRangeEndFrame"))
 	int32 CustomPlaybackRangeEndFrame;
 
+	/** Start the timecode at a specific value, rather than the value coming from the Level Sequence. Only applicable to output formats that support timecode. */
+	UPROPERTY(EditAnywhere, BlueprintReadWrite, Category = "Timecode", meta = (EditCondition = "bOverride_CustomTimecodeStart"))
+	FTimecode CustomTimecodeStart;
+
+	/** Whether the embedded timecode track should be written using drop-frame format. Only applicable to output formats that support timecode, and when the sequence framerate is 29.97. */
+	UPROPERTY(EditAnywhere, BlueprintReadWrite, Category = "Timecode", DisplayName = "Use DF Timecode if 29.97 FPS", meta = (EditCondition = "bOverride_bDropFrameTimecode"))
+	bool bDropFrameTimecode;
+
 	/**
 	 * Determines how versioning should be handled (Auto Version, Version Number, etc.).
 	 */
