MoviePipeline: Added additional telemetry for the Avid DNxHR and Apple ProRes nodes.

#rnx

diff --git a/Engine/Plugins/Media/AppleProResMedia/Source/AppleProResMedia/Private/MovieGraphAppleProResNode.cpp b/Engine/Plugins/Media/AppleProResMedia/Source/AppleProResMedia/Private/MovieGraphAppleProResNode.cpp
--- a/Engine/Plugins/Media/AppleProResMedia/Source/AppleProResMedia/Private/MovieGraphAppleProResNode.cpp
+++ b/Engine/Plugins/Media/AppleProResMedia/Source/AppleProResMedia/Private/MovieGraphAppleProResNode.cpp
@@ -10,6 +10,7 @@
 #include "Graph/Nodes/MovieGraphGlobalOutputSettingNode.h"
 #include "ImageWriteTask.h"
 #include "MoviePipelineImageQuantization.h"
+#include "MoviePipelineTelemetry.h"
 #include "MovieRenderPipelineCoreModule.h"
 #include "Styling/AppStyle.h"
 
@@ -194,4 +195,9 @@ bool UMovieGraphAppleProResNode::IsAudioSupported() const
 {
 	// The current ProRes SDK does not support audio so we don't write audio.
 	return false;
-}
\ No newline at end of file
+}
+
+void UMovieGraphAppleProResNode::UpdateTelemetry(FMoviePipelineShotRenderTelemetry* InTelemetry) const
+{
+	InTelemetry->bUsesProRes = true;
+}
diff --git a/Engine/Plugins/Media/AppleProResMedia/Source/AppleProResMedia/Private/MovieGraphAppleProResNode.h b/Engine/Plugins/Media/AppleProResMedia/Source/AppleProResMedia/Private/MovieGraphAppleProResNode.h
--- a/Engine/Plugins/Media/AppleProResMedia/Source/AppleProResMedia/Private/MovieGraphAppleProResNode.h
+++ b/Engine/Plugins/Media/AppleProResMedia/Source/AppleProResMedia/Private/MovieGraphAppleProResNode.h
@@ -36,6 +36,10 @@ protected:
 	virtual bool IsAudioSupported() const override;
 	// ~UMovieGraphVideoOutputNode Interface
 
+	// UMovieGraphSettingNode Interface
+	virtual void UpdateTelemetry(FMoviePipelineShotRenderTelemetry* InTelemetry) const override;
+	// ~UMovieGraphSettingNode Interface
+
 protected:
 	struct FProResWriter : public MovieRenderGraph::IVideoCodecWriter
 	{
diff --git a/Engine/Plugins/Media/AvidDNxMedia/Source/Source/Private/MovieGraphAvidDNxHRNode.cpp b/Engine/Plugins/Media/AvidDNxMedia/Source/Source/Private/MovieGraphAvidDNxHRNode.cpp
--- a/Engine/Plugins/Media/AvidDNxMedia/Source/Source/Private/MovieGraphAvidDNxHRNode.cpp
+++ b/Engine/Plugins/Media/AvidDNxMedia/Source/Source/Private/MovieGraphAvidDNxHRNode.cpp
@@ -4,6 +4,7 @@
 
 #include "ImageWriteTask.h"
 #include "MoviePipelineImageQuantization.h"
+#include "MoviePipelineTelemetry.h"
 #include "AvidDNxEncoder/AvidDNxEncoder.h"
 #include "Graph/MovieGraphBlueprintLibrary.h"
 #include "Graph/MovieGraphConfig.h"
@@ -211,4 +212,9 @@ bool UMovieGraphAvidDNxHRNode::IsAudioSupported() const
 	return false;
 }
 
-#endif // WITH_EDITOR
\ No newline at end of file
+#endif // WITH_EDITOR
+
+void UMovieGraphAvidDNxHRNode::UpdateTelemetry(FMoviePipelineShotRenderTelemetry* InTelemetry) const
+{
+	InTelemetry->bUsesAvid = true;
+}
diff --git a/Engine/Plugins/Media/AvidDNxMedia/Source/Source/Private/MovieGraphAvidDNxHRNode.h b/Engine/Plugins/Media/AvidDNxMedia/Source/Source/Private/MovieGraphAvidDNxHRNode.h
--- a/Engine/Plugins/Media/AvidDNxMedia/Source/Source/Private/MovieGraphAvidDNxHRNode.h
+++ b/Engine/Plugins/Media/AvidDNxMedia/Source/Source/Private/MovieGraphAvidDNxHRNode.h
@@ -38,6 +38,10 @@ protected:
 	virtual bool IsAudioSupported() const override;
 	// ~UMovieGraphVideoOutputNode Interface
 
+	// UMovieGraphSettingNode Interface
+	virtual void UpdateTelemetry(FMoviePipelineShotRenderTelemetry* InTelemetry) const override;
+	// ~UMovieGraphSettingNode Interface
+
 protected:
 	struct FAvidWriter : public MovieRenderGraph::IVideoCodecWriter
 	{
diff --git a/Engine/Plugins/MovieScene/MovieRenderPipeline/Source/MovieRenderPipelineCore/Public/MoviePipelineTelemetry.h b/Engine/Plugins/MovieScene/MovieRenderPipeline/Source/MovieRenderPipelineCore/Public/MoviePipelineTelemetry.h
--- a/Engine/Plugins/MovieScene/MovieRenderPipeline/Source/MovieRenderPipelineCore/Public/MoviePipelineTelemetry.h
+++ b/Engine/Plugins/MovieScene/MovieRenderPipeline/Source/MovieRenderPipelineCore/Public/MoviePipelineTelemetry.h
@@ -24,6 +24,8 @@ struct MOVIERENDERPIPELINECORE_API FMoviePipelineShotRenderTelemetry
 	bool bUsesSubgraphs = false;
 	bool bUsesPPMs = false;
 	bool bUsesAudio = false;
+	bool bUsesAvid = false;
+	bool bUsesProRes = false;
 	int32 ResolutionX = 0;
 	int32 ResolutionY = 0;
 	int32 HandleFrameCount = 0;
