MoviePipeline: Added support for embedding audio recording into Apple ProRes mov files. This does not improve the quality of recorded audio (ie: there are still gaps between shots where enigne warm-ups take place) but now embeds the recorded audio into the .mov file as uncompressed lpcm audio.

diff --git a/Engine/Plugins/Media/AppleProResMedia/Source/AppleProResMedia/Private/AppleProResEncoder/AppleProResEncoder.cpp b/Engine/Plugins/Media/AppleProResMedia/Source/AppleProResMedia/Private/AppleProResEncoder/AppleProResEncoder.cpp
--- a/Engine/Plugins/Media/AppleProResMedia/Source/AppleProResMedia/Private/AppleProResEncoder/AppleProResEncoder.cpp
+++ b/Engine/Plugins/Media/AppleProResMedia/Source/AppleProResMedia/Private/AppleProResEncoder/AppleProResEncoder.cpp
@@ -301,10 +301,10 @@ bool FAppleProResEncoder::Initialize()
 	// ProRes does Numerator/Denominator like us but reversed - Our time scale is frames-per-second, they are seconds-per-frame so we swap math.
 	PRTimeScale ProResTimescale = Options.FrameRate.Numerator;
 
-	CurrentTime.value = Options.FrameNumberOffset * Options.FrameRate.Denominator;
-	CurrentTime.timescale = Options.FrameRate.Numerator;
-	CurrentTime.flags = kPRTimeFlags_Valid;
-	CurrentTime.epoch = 0;
+	CurrentVideoTime.value = Options.FrameNumberOffset * Options.FrameRate.Denominator;
+	CurrentVideoTime.timescale = Options.FrameRate.Numerator;
+	CurrentVideoTime.flags = kPRTimeFlags_Valid;
+	CurrentVideoTime.epoch = 0;
 
 	// Set the time scale in the header. Must match the samples we write later and must be before header is written.
 	Status = ProResFileWriterSetMovieTimescale(FileWriter, ProResTimescale);
@@ -337,7 +337,15 @@ bool FAppleProResEncoder::Initialize()
 			UE_LOG(LogAppleProResMedia, Error, TEXT("Failed to add audio track to file writer at path: \"%s\"."), *Options.OutputFilename);
 			return false;
 		}
-		Status = ProResFileWriterSetTrackMediaTimescale(FileWriter, AudioTrackId, ProResTimescale);
+		
+		// Sound tracks must match the sample rate, not frame rate.
+		if(!ensureMsgf(Options.AudioSampleRate > 0, TEXT("AudioSampleRate must be greater than zero!")))
+		{
+			return false;
+		}
+		
+		PRTimeScale ProResAudioTimescale = Options.AudioSampleRate;
+		Status = ProResFileWriterSetTrackMediaTimescale(FileWriter, AudioTrackId, ProResAudioTimescale);
 		if (Status != 0)
 		{
 			UE_LOG(LogAppleProResMedia, Error, TEXT("Failed to set audio track timescale to file writer at path: \"%s\"."), *Options.OutputFilename);
@@ -384,7 +392,7 @@ bool FAppleProResEncoder::Initialize()
 		}
 	}
 
-	Status = ProResFileWriterBeginSession(FileWriter, CurrentTime);
+	Status = ProResFileWriterBeginSession(FileWriter, CurrentVideoTime);
 	if (Status != 0)
 	{
 		UE_LOG(LogAppleProResMedia, Error, TEXT("Failed to begin session with file writer at path: \"%s\"."), *Options.OutputFilename);
@@ -478,19 +486,28 @@ bool FAppleProResEncoder::InitializeTimecodeTrack()
 
 bool FAppleProResEncoder::InitializeAudioTrack()
 {
-	AudioStreamBasicDescription AudioSettings;
-	AudioSettings.mSampleRate = 48000;
+	AudioStreamBasicDescription AudioSettings = {};
+	AudioSettings.mSampleRate = Options.AudioSampleRate; 
 	AudioSettings.mFormatID = kAudioFormatLinearPCM;
-	AudioSettings.mFormatFlags = kAudioFormatFlagIsSignedInteger & kAudioFormatFlagIsPacked;
-
-	// Stubbed out but untested - Let's assume a packet is 48000? And 24 fps. = 2000 samples per packet.
-	AudioSettings.mBytesPerPacket = 48000 * 2 * 2; // 2 channels of signed int16 PCM
-	AudioSettings.mFramesPerPacket = 24;
-	AudioSettings.mBytesPerFrame = 2000 * 2 * 2; // 2 channels @ int16, for 2000 samples per frame.
-	AudioSettings.mChannelsPerFrame = 2;
-	AudioSettings.mBitsPerChannel = 16;
+	AudioSettings.mFormatFlags = kAudioFormatFlagIsSignedInteger | kAudioFormatFlagIsPacked;
+
+	// To figure out how much audio data there is per sample, we take our AudioSampleRate (48000) and divide it by our framerate, ie: 48'000 / 30 = 1600
+	AudioSamplesPerFrame = FMath::FloorToInt(Options.AudioSampleRate * Options.FrameRate.AsInterval());
+
+	// Audio is uncompressed lpcm, and for Apple ProRes 1 frame is 1 sample across all channels, and 1 packet is 1 frame in compressed files.
+	// Then when it's submitted, we can submit one video-frame's worth of data at a time.
+	AudioSettings.mBytesPerPacket = Options.AudioChannelCount * sizeof(int16); // 2 channels of signed int16 PCM
+	AudioSettings.mFramesPerPacket = 1;
+	AudioSettings.mBytesPerFrame = Options.AudioChannelCount * sizeof(int16);
+	AudioSettings.mChannelsPerFrame = Options.AudioChannelCount;
+	AudioSettings.mBitsPerChannel = sizeof(int16)*8;
 	AudioSettings.mReserved = 0; // Padding
 
+	CurrentAudioTime.value = Options.FrameNumberOffset * Options.FrameRate.Denominator;
+	CurrentAudioTime.timescale = Options.AudioSampleRate;
+	CurrentAudioTime.flags = kPRTimeFlags_Valid;
+	CurrentAudioTime.epoch = 0;
+
 
 	PRStatus Status = ProResAudioFormatDescriptionCreate(
 		&AudioSettings,
@@ -535,8 +552,8 @@ void FAppleProResEncoder::Finalize()
 		}
 
 		// Add one more sample onto the end 
-		CurrentTime.value += Options.FrameRate.Denominator;
-		Status = ProResFileWriterEndSession(FileWriter, CurrentTime);
+		CurrentVideoTime.value += Options.FrameRate.Denominator;
+		Status = ProResFileWriterEndSession(FileWriter, CurrentVideoTime);
 		if (Status != 0)
 		{
 			UE_LOG(LogAppleProResMedia, Error, TEXT("Failed to End Writer Session for file writer at path: \"%s\"."), *Options.OutputFilename);
@@ -604,9 +621,9 @@ bool FAppleProResEncoder::WriteFrame(const FImagePixelData* InPixelData)
 	SampleTimingInfo.duration.timescale = Options.FrameRate.Numerator;
 	SampleTimingInfo.duration.flags = kPRTimeFlags_Valid;
 	SampleTimingInfo.duration.epoch = 0;
-	SampleTimingInfo.timeStamp = CurrentTime;
+	SampleTimingInfo.timeStamp = CurrentVideoTime;
 
-	CurrentTime.value += Options.FrameRate.Denominator;
+	CurrentVideoTime.value += Options.FrameRate.Denominator;
 
 	// Build the Encoder Params
 	PREncodingParams EncodingParams;
@@ -690,7 +707,46 @@ bool FAppleProResEncoder::WriteFrame(const FImagePixelData* InPixelData)
 
 bool FAppleProResEncoder::WriteAudioSample(const TArrayView<int16>& InAudioSamples)
 {
-	return false;
+	if (!ensureMsgf(bInitialized || bFinalized, TEXT("WriteFrame should not be called if not initialized or after finalize! Initialized: %d Finalized: %d"), bInitialized, bFinalized))
+	{
+		return false;
+	}
+
+	int32 SampleCountPerChannel = InAudioSamples.Num() / Options.AudioChannelCount;
+	int32 NumFramesOfData = SampleCountPerChannel / AudioSamplesPerFrame;
+
+	int16* SourceDataPtr = (int16*)InAudioSamples.GetData();
+	for (int32 SampleIdx = 0; SampleIdx < NumFramesOfData; SampleIdx++)
+	{
+		// The number of bytes per frame is channel count (due to interleaved audio) * sizeof(int16) (signed 16-bit audio) * AudioSamplesPerFrame.
+		const size_t BufferSize = AudioSamplesPerFrame * Options.AudioChannelCount * sizeof(int16);
+		int32 NumSamplesInBuffer = AudioSamplesPerFrame;
+		PRStatus Status = ProResFileWriterAddAudioSampleBufferToTrack(
+			FileWriter,
+			AudioTrackId,
+			(void*)SourceDataPtr,
+			BufferSize,
+			/*Deallocator*/ nullptr,
+			AudioFormatDescription,
+			NumSamplesInBuffer,
+			CurrentAudioTime
+			);
+
+		// Increment our source pointer to the next frame's data.
+		SourceDataPtr += AudioSamplesPerFrame * Options.AudioChannelCount;
+
+		// Increment our CurrentAudioTime (which is in hz, ie: 48'000) by 1 frame's worth of samples.
+		CurrentAudioTime.value += AudioSamplesPerFrame;
+
+		if (Status != 0)
+		{
+			UE_LOG(LogAppleProResMedia, Error, TEXT("Failed to add sample to Audio Track for file writer at path: \"%s\"."), *Options.OutputFilename);
+			return false;
+		}
+
+	}
+
+	return true;
 }
 
 const FAppleProResEncoderOptions& FAppleProResEncoder::GetOptions() const
diff --git a/Engine/Plugins/Media/AppleProResMedia/Source/AppleProResMedia/Private/MovieGraphAppleProResNode.cpp b/Engine/Plugins/Media/AppleProResMedia/Source/AppleProResMedia/Private/MovieGraphAppleProResNode.cpp
--- a/Engine/Plugins/Media/AppleProResMedia/Source/AppleProResMedia/Private/MovieGraphAppleProResNode.cpp
+++ b/Engine/Plugins/Media/AppleProResMedia/Source/AppleProResMedia/Private/MovieGraphAppleProResNode.cpp
@@ -13,6 +13,7 @@
 #include "MoviePipelineTelemetry.h"
 #include "MovieRenderPipelineCoreModule.h"
 #include "Styling/AppStyle.h"
+#include "SampleBuffer.h"
 
 UMovieGraphAppleProResNode::UMovieGraphAppleProResNode()
 	: Quality(EAppleProResEncoderCodec::ProRes_422LT)
@@ -80,10 +81,10 @@ TUniquePtr<MovieRenderGraph::IVideoCodecWriter> UMovieGraphAppleProResNode::Init
 	Options.bDropFrameTimecode = OutputSetting->bDropFrameTimecode;
 
 	// If OCIO is enabled, don't do additional color conversion
-	Options.bConvertToSrgb = !(EvaluatedNode->bOverride_OCIOConfiguration && EvaluatedNode->OCIOConfiguration.bIsEnabled && InInitializationContext.bAllowOCIO);
 	
 	TUniquePtr<FProResWriter> NewWriter = MakeUnique<FProResWriter>();
 	NewWriter->Writer = MakeUnique<FAppleProResEncoder>(Options);
+	NewWriter->bSkipColorConversions = (EvaluatedNode->bOverride_OCIOConfiguration && EvaluatedNode->OCIOConfiguration.bIsEnabled && InInitializationContext.bAllowOCIO);
 
 	CachedPipeline = InInitializationContext.Pipeline;
 	
@@ -104,6 +105,14 @@ bool UMovieGraphAppleProResNode::Initialize_EncodeThread(MovieRenderGraph::IVide
 
 void UMovieGraphAppleProResNode::WriteFrame_EncodeThread(MovieRenderGraph::IVideoCodecWriter* InWriter, FImagePixelData* InPixelData, TArray<FMovieGraphPassData>&& InCompositePasses, TObjectPtr<UMovieGraphEvaluatedConfig> InEvaluatedConfig, const FString& InBranchName)
 {
+	FProResWriter* CodecWriter = static_cast<FProResWriter*>(InWriter);
+
+	// If the writer was not initialized, don't try to finalize anything.
+	if (!CodecWriter->Writer)
+	{
+		return;
+	}
+
 	bool bIncludeCDOs = false;
 	constexpr bool bExactMatch = true;
 	const UMovieGraphAppleProResNode* EvaluatedNode = Cast<UMovieGraphAppleProResNode>(
@@ -114,8 +123,6 @@ void UMovieGraphAppleProResNode::WriteFrame_EncodeThread(MovieRenderGraph::IVide
 	const UMovieGraphGlobalOutputSettingNode* OutputSettingNode =
 		InEvaluatedConfig->GetSettingForBranch<UMovieGraphGlobalOutputSettingNode>(GlobalsPinName, bIncludeCDOs, bExactMatch);
 	
-	const FProResWriter* CodecWriter = static_cast<FProResWriter*>(InWriter);
-	
 	const UE::MovieGraph::FMovieGraphSampleState* GraphPayload = InPixelData->GetPayload<UE::MovieGraph::FMovieGraphSampleState>();
 	
 	// Translate our Movie Pipeline specific payload to a ProRes Encoder specific payload.
@@ -135,8 +142,9 @@ void UMovieGraphAppleProResNode::WriteFrame_EncodeThread(MovieRenderGraph::IVide
 		ProResPayload->ReferenceFrameNumber = GraphPayload->TraversalContext.Time.RootFrameNumber.Value;
 	}
 
-	// ProRes can handle quantization internally but expects sRGB to be applied to the incoming data.
-	const TUniquePtr<FImagePixelData> QuantizedPixelData = UE::MoviePipeline::QuantizeImagePixelDataToBitDepth(InPixelData, 8, ProResPayload, CodecWriter->Writer->GetOptions().bConvertToSrgb);
+	// ProRes can handle quantization internally but expects sRGB to be applied to the incoming data (Rec 709)
+	const bool bConvertToSrgb = !CodecWriter->bSkipColorConversions;
+	const TUniquePtr<FImagePixelData> QuantizedPixelData = UE::MoviePipeline::QuantizeImagePixelDataToBitDepth(InPixelData, 8, ProResPayload, bConvertToSrgb);
 
 	TArray<FPixelPreProcessor> PixelPreProcessors;
 
@@ -176,9 +184,55 @@ void UMovieGraphAppleProResNode::WriteFrame_EncodeThread(MovieRenderGraph::IVide
 
 void UMovieGraphAppleProResNode::BeginFinalize_EncodeThread(MovieRenderGraph::IVideoCodecWriter* InWriter)
 {
-	return;
+	if (!CachedPipeline.IsValid())
+	{
+		return;
+	}
+
+	const MoviePipeline::FAudioState& AudioData = CachedPipeline->GetAudioRendererInstance()->GetAudioState();
+	FProResWriter* CodecWriter = static_cast<FProResWriter*>(InWriter);
+
+	// If the writer was not initialized, don't try to finalize anything.
+	if (!CodecWriter->Writer)
+	{
+		return;
+	}
+
+	if (!CodecWriter->Writer->IsInitialized())
+	{
+		return;
+	}
+
+	for (const TPair<int32, MovieRenderGraph::IVideoCodecWriter::FLightweightSourceData>& SourceData : CodecWriter->LightweightSourceData)
+	{
+		if (!AudioData.FinishedSegments.IsValidIndex(SourceData.Key))
+		{
+			UE_LOG(LogMovieRenderPipeline, Warning, TEXT("Invalid shot index was requested for audio data, skipping audio writes."));
+			continue;
+		}
+
+		// Look up the audio segment for this shot
+		const MoviePipeline::FAudioState::FAudioSegment& AudioSegment = AudioData.FinishedSegments[SourceData.Key];
+
+		// Audio data isn't very sample accurate at this point, so we may have generated slightly more (or less) audio than we expect for
+		// the number of frames, so we're actually going to trim down the view of data we provide to match the number of frames rendered,
+		// to avoid any excess audio after the end of the video.
+		const int32 NumFrames = SourceData.Value.SubmittedFrameCount;
+
+		// ToDo: This is possibly dropping fractions of a sample (ie: 1/48,000th) if the audio sample rate can't be evenly divisible by
+		// the frame rate.
+		const int32 SamplesPerFrame = AudioSegment.SampleRate * CodecWriter->Writer->GetOptions().FrameRate.AsInterval();
+		int32 ExpectedSampleCount = NumFrames * SamplesPerFrame * AudioSegment.NumChannels;
+
+		ExpectedSampleCount = FMath::Min(ExpectedSampleCount, AudioSegment.SegmentData.Num());
+		Audio::TSampleBuffer<int16> SampleBuffer = Audio::TSampleBuffer<int16>(AudioSegment.SegmentData.GetData(), ExpectedSampleCount, AudioSegment.NumChannels, AudioSegment.SampleRate);
+
+		const TArrayView<int16> SampleData = SampleBuffer.GetArrayView();
+		CodecWriter->Writer->WriteAudioSample(SampleData);
+	}
 }
 
+
 void UMovieGraphAppleProResNode::Finalize_EncodeThread(MovieRenderGraph::IVideoCodecWriter* InWriter)
 {
 	// Write to disk
diff --git a/Engine/Plugins/Media/AppleProResMedia/Source/AppleProResMedia/Private/MovieGraphAppleProResNode.h b/Engine/Plugins/Media/AppleProResMedia/Source/AppleProResMedia/Private/MovieGraphAppleProResNode.h
--- a/Engine/Plugins/Media/AppleProResMedia/Source/AppleProResMedia/Private/MovieGraphAppleProResNode.h
+++ b/Engine/Plugins/Media/AppleProResMedia/Source/AppleProResMedia/Private/MovieGraphAppleProResNode.h
@@ -1,4 +1,4 @@
-﻿// Copyright Epic Games, Inc. All Rights Reserved.
+// Copyright Epic Games, Inc. All Rights Reserved.
 
 #pragma once
 
@@ -43,6 +43,7 @@ protected:
 protected:
 	struct FProResWriter : public MovieRenderGraph::IVideoCodecWriter
 	{
+		bool bSkipColorConversions = false;
 		TUniquePtr<FAppleProResEncoder> Writer;
 	};
 	
diff --git a/Engine/Plugins/Media/AppleProResMedia/Source/AppleProResMedia/Private/MoviePipelineAppleProResOutput.cpp b/Engine/Plugins/Media/AppleProResMedia/Source/AppleProResMedia/Private/MoviePipelineAppleProResOutput.cpp
--- a/Engine/Plugins/Media/AppleProResMedia/Source/AppleProResMedia/Private/MoviePipelineAppleProResOutput.cpp
+++ b/Engine/Plugins/Media/AppleProResMedia/Source/AppleProResMedia/Private/MoviePipelineAppleProResOutput.cpp
@@ -66,8 +66,9 @@ void UMoviePipelineAppleProResOutput::WriteFrame_EncodeThread(MovieRenderPipelin
 	// This is the frame number on the global time, can have overlaps (between encoders) or repeats when using handle frames/slowmo.
 	ProResPayload->ReferenceFrameNumber = PipelinePayload->SampleState.OutputState.SourceFrameNumber;
 
-	// ProRes can handle quantization internally but expects sRGB to be applied to the incoming data.
-	TUniquePtr<FImagePixelData> QuantizedPixelData = UE::MoviePipeline::QuantizeImagePixelDataToBitDepth(InPixelData, 16, ProResPayload, InWriter->bConvertToSrgb);
+	// ProRes can handle either 16 or 8 bit input internally, but expects Rec709 input which has a sRGB curve applied.
+	constexpr bool bConvertTosRGB = true;
+	TUniquePtr<FImagePixelData> QuantizedPixelData = UE::MoviePipeline::QuantizeImagePixelDataToBitDepth(InPixelData, 16, ProResPayload, bConvertTosRGB);
 
 	// Do a quick composite of renders/burn-ins.
 	TArray<FPixelPreProcessor> PixelPreProcessors;
@@ -105,17 +106,44 @@ void UMoviePipelineAppleProResOutput::BeginFinalize_EncodeThread(MovieRenderPipe
 	const MoviePipeline::FAudioState& AudioData = GetPipeline()->GetAudioState();
 	FProResWriter* CodecWriter = static_cast<FProResWriter*>(InWriter);
 
-	// The AppleProResEncoder implementation does not currently support encoding audio.
-	/*for (const MoviePipeline::FAudioState::FAudioSegment& AudioSegment : AudioData.FinishedSegments)
+	// If the writer was not initialized, don't try to finalize anything.
+	if (!CodecWriter->Writer)
 	{
-		double StartTime = FPlatformTime::Seconds();
-		Audio::TSampleBuffer<int16> SampleBuffer = Audio::TSampleBuffer<int16>(AudioSegment.SegmentData.GetData(), AudioSegment.SegmentData.Num(), AudioSegment.NumChannels, AudioSegment.SampleRate);
-		UE_LOG(LogMovieRenderPipeline, Log, TEXT("Audio Segment took %f seconds to convert to a sample buffer."), (FPlatformTime::Seconds() - StartTime));
+		return;
+	}
+
+	if (!CodecWriter->Writer->IsInitialized())
+	{
+		return;
+	}
+
+	for (const TPair<int32, MovieRenderPipeline::IVideoCodecWriter::FLightweightSourceData>& SourceData : CodecWriter->LightweightSourceData)
+	{
+		if (!AudioData.FinishedSegments.IsValidIndex(SourceData.Key))
+		{
+			UE_LOG(LogMovieRenderPipeline, Warning, TEXT("Invalid shot index was requested for audio data, skipping audio writes."));
+			continue;
+		}
+
+		// Look up the audio segment for this shot
+		const MoviePipeline::FAudioState::FAudioSegment& AudioSegment = AudioData.FinishedSegments[SourceData.Key];
+
+		// Audio data isn't very sample accurate at this point, so we may have generated slightly more (or less) audio than we expect for
+		// the number of frames, so we're actually going to trim down the view of data we provide to match the number of frames rendered,
+		// to avoid any excess audio after the end of the video.
+		const int32 NumFrames = SourceData.Value.SubmittedFrameCount;
+
+		// ToDo: This is possibly dropping fractions of a sample (ie: 1/48,000th) if the audio sample rate can't be evenly divisible by
+		// the frame rate.
+		const int32 SamplesPerFrame = AudioSegment.SampleRate * CodecWriter->Writer->GetOptions().FrameRate.AsInterval();
+		int32 ExpectedSampleCount = NumFrames * SamplesPerFrame * AudioSegment.NumChannels;
 
-		const TArrayView<int16> SampleData = SampleBuffer.GetArrayView(); // MakeArrayView<int16>(SampleBuffer.GetData(), SampleBuffer.GetNumSamples());
+		ExpectedSampleCount = FMath::Min(ExpectedSampleCount, AudioSegment.SegmentData.Num());
+		Audio::TSampleBuffer<int16> SampleBuffer = Audio::TSampleBuffer<int16>(AudioSegment.SegmentData.GetData(), ExpectedSampleCount, AudioSegment.NumChannels, AudioSegment.SampleRate);
+
+		const TArrayView<int16> SampleData = SampleBuffer.GetArrayView();
 		CodecWriter->Writer->WriteAudioSample(SampleData);
-	}*/
-	
+	}
 }
 
 void UMoviePipelineAppleProResOutput::Finalize_EncodeThread(MovieRenderPipeline::IVideoCodecWriter* InWriter)
diff --git a/Engine/Plugins/Media/AppleProResMedia/Source/AppleProResMedia/Private/MoviePipelineAppleProResOutput.h b/Engine/Plugins/Media/AppleProResMedia/Source/AppleProResMedia/Private/MoviePipelineAppleProResOutput.h
--- a/Engine/Plugins/Media/AppleProResMedia/Source/AppleProResMedia/Private/MoviePipelineAppleProResOutput.h
+++ b/Engine/Plugins/Media/AppleProResMedia/Source/AppleProResMedia/Private/MoviePipelineAppleProResOutput.h
@@ -31,7 +31,7 @@ protected:
 	virtual void BeginFinalize_EncodeThread(MovieRenderPipeline::IVideoCodecWriter* InWriter);
 	virtual void Finalize_EncodeThread(MovieRenderPipeline::IVideoCodecWriter* InWriter);
 	virtual const TCHAR* GetFilenameExtension() const override { return TEXT("mov"); }
-	virtual bool IsAudioSupported() const { return false; }
+	virtual bool IsAudioSupported() const { return true; }
 	// ~UMoviePipelineVideoOutputBase Interface
 
 	// UMoviePipelineOutputBase Interface
diff --git a/Engine/Plugins/Media/AppleProResMedia/Source/AppleProResMedia/Public/AppleProResEncoder/AppleProResEncoder.h b/Engine/Plugins/Media/AppleProResMedia/Source/AppleProResMedia/Public/AppleProResEncoder/AppleProResEncoder.h
--- a/Engine/Plugins/Media/AppleProResMedia/Source/AppleProResMedia/Public/AppleProResEncoder/AppleProResEncoder.h
+++ b/Engine/Plugins/Media/AppleProResMedia/Source/AppleProResMedia/Public/AppleProResEncoder/AppleProResEncoder.h
@@ -69,6 +69,8 @@ struct FAppleProResEncoderOptions
 		, bWriteAlpha(false)
 		, bDropFrameTimecode(false)
 		, FrameNumberOffset(0)
+		, AudioSampleRate(48000)
+		, AudioChannelCount(2)
 	{}
 
 	/** The absolute path on disk to try and save the video file to. */
@@ -90,6 +92,7 @@ struct FAppleProResEncoderOptions
 	EAppleProResEncoderColorPrimaries ColorPrimaries;
 
 	/** Whether the data should be converted to sRGB before being sent to the encoder. Should not be used when OCIO is active. */
+	UE_DEPRECATED(5.6, "The encoder currently always assumes Rec 709 input, and never used this parameter.")
 	bool bConvertToSrgb;
 
 	/** Which scan mode do we use? Only Progressive is tested right now. */
@@ -106,6 +109,12 @@ struct FAppleProResEncoderOptions
 
 	/** The number of frames to offset the Timecode track by. */
 	int32 FrameNumberOffset;
+
+	/** The sample rate of the incoming audio data. Only tested with 48'000. */
+	uint32 AudioSampleRate;
+
+	/** The number of audio channels in the encoding data. Only tested with 2. */
+	uint32 AudioChannelCount;
 };
 
 class APPLEPRORESMEDIA_API FAppleProResEncoder
@@ -135,6 +144,8 @@ public:
 	/** Gets the options that the encoder was initialized with. */
 	const FAppleProResEncoderOptions& GetOptions() const;
 
+	bool IsInitialized() const { return bInitialized; }
+
 private:
 	bool InitializeVideoTrack();
 	bool InitializeTimecodeTrack();
@@ -160,8 +171,13 @@ private:
 	/** If valid, we own the memory and need to free it. */
 	PREncoderRef Encoder;
 
-	PRTime CurrentTime;
+	/** Time in the file for video tracks. Incremented by 1 frame for each new video sample. Video samples come in one at a time so this keeps track of how many have been added. */
+	PRTime CurrentVideoTime;
+	/** Time in the file for audio tracks. Uses a different timescale than Video, as it's in AudioSampleRate. */
+	PRTime CurrentAudioTime;
+
 
 	int32 MaxCompressedFrameSize;
 	int32 TargetCompressedFrameSize;
+	uint32 AudioSamplesPerFrame;
 };
\ No newline at end of file
diff --git a/Engine/Plugins/MovieScene/MovieRenderPipeline/Source/MovieRenderPipelineCore/Private/Graph/Nodes/MovieGraphVideoOutputNode.cpp b/Engine/Plugins/MovieScene/MovieRenderPipeline/Source/MovieRenderPipelineCore/Private/Graph/Nodes/MovieGraphVideoOutputNode.cpp
--- a/Engine/Plugins/MovieScene/MovieRenderPipeline/Source/MovieRenderPipelineCore/Private/Graph/Nodes/MovieGraphVideoOutputNode.cpp
+++ b/Engine/Plugins/MovieScene/MovieRenderPipeline/Source/MovieRenderPipelineCore/Private/Graph/Nodes/MovieGraphVideoOutputNode.cpp
@@ -91,6 +91,10 @@ void UMovieGraphVideoOutputNode::OnReceiveImageDataImpl(UMovieGraphPipeline* InP
 		FImagePixelData* RawRenderPassData = RenderPassData.Value.Get();
 		const UE::MovieGraph::FMovieGraphSampleState* Payload = RawRenderPassData->GetPayload<UE::MovieGraph::FMovieGraphSampleState>();
 
+		// We can have multiple shots in one file, so we find/add the shot in case this is a new shot, and the counter will start over at zero.
+		MovieRenderGraph::IVideoCodecWriter::FLightweightSourceData& LightweightData = OutputWriter->CodecWriter->LightweightSourceData.FindOrAdd(Payload->TraversalContext.ShotIndex);
+		LightweightData.SubmittedFrameCount++;
+
 		// Pixel data pointer for the cropped result to ensure it stays in memory while the image write is processed
 		TUniquePtr<FImagePixelData> CroppedPixelData;
 		const bool bIsCropRectValid = !Payload->CropRectangle.IsEmpty();
diff --git a/Engine/Plugins/MovieScene/MovieRenderPipeline/Source/MovieRenderPipelineCore/Private/MoviePipelineVideoOutputBase.cpp b/Engine/Plugins/MovieScene/MovieRenderPipeline/Source/MovieRenderPipelineCore/Private/MoviePipelineVideoOutputBase.cpp
--- a/Engine/Plugins/MovieScene/MovieRenderPipeline/Source/MovieRenderPipelineCore/Private/MoviePipelineVideoOutputBase.cpp
+++ b/Engine/Plugins/MovieScene/MovieRenderPipeline/Source/MovieRenderPipelineCore/Private/MoviePipelineVideoOutputBase.cpp
@@ -273,6 +273,10 @@ void UMoviePipelineVideoOutputBase::OnReceiveImageDataImpl(FMoviePipelineMergerO
 		UMoviePipelineColorSetting* ColorSetting = GetPipeline()->GetPipelinePrimaryConfig()->FindSetting<UMoviePipelineColorSetting>();
 		OutputWriter->Get<0>()->bConvertToSrgb = !(ColorSetting && ColorSetting->OCIOConfiguration.bIsEnabled);
 
+		// We can have multiple shots in one file, so we find/add the shot in case this is a new shot, and the counter will start over at zero.
+		MovieRenderPipeline::IVideoCodecWriter::FLightweightSourceData& LightweightData = OutputWriter->Get<0>()->LightweightSourceData.FindOrAdd(Payload->SampleState.OutputState.ShotIndex);
+		LightweightData.SubmittedFrameCount++;
+		
 		//FGraphEventRef Event = Task.Execute([this, OutputWriter, RawRenderPassData]
 		//	{
 			if (RenderPassData.Key.Name == TEXT("FinalImage"))
diff --git a/Engine/Plugins/MovieScene/MovieRenderPipeline/Source/MovieRenderPipelineCore/Public/Graph/Nodes/MovieGraphVideoOutputNode.h b/Engine/Plugins/MovieScene/MovieRenderPipeline/Source/MovieRenderPipelineCore/Public/Graph/Nodes/MovieGraphVideoOutputNode.h
--- a/Engine/Plugins/MovieScene/MovieRenderPipeline/Source/MovieRenderPipelineCore/Public/Graph/Nodes/MovieGraphVideoOutputNode.h
+++ b/Engine/Plugins/MovieScene/MovieRenderPipeline/Source/MovieRenderPipelineCore/Public/Graph/Nodes/MovieGraphVideoOutputNode.h
@@ -1,4 +1,4 @@
-﻿// Copyright Epic Games, Inc. All Rights Reserved.
+// Copyright Epic Games, Inc. All Rights Reserved.
 #pragma once
 
 #include "Async/AsyncWork.h"
@@ -20,6 +20,16 @@ namespace MovieRenderGraph
 		 * due to the file existing on disk after the first frame comes in, and then the next one de-duplicating to one more than that.
 		 */
 		FString StableFileName;
+
+		// This array contains the indexes of each shot that contributed to this writer. This is needed so that after a shot/render, when
+		// all we have is the completed blocks of audio samples, we can figure out which blocks go to which video clips. Because of how you can
+		// split videos up by either shot, or by the whole sequence, there isn't an obvious mapping, so this array will tell the audio writing
+		// system which audio blocks should be sent to the shot.
+		struct FLightweightSourceData
+		{
+			int32 SubmittedFrameCount;
+		};
+		TMap<int32, FLightweightSourceData> LightweightSourceData;
 	};
 }
 
diff --git a/Engine/Plugins/MovieScene/MovieRenderPipeline/Source/MovieRenderPipelineCore/Public/MoviePipelineVideoOutputBase.h b/Engine/Plugins/MovieScene/MovieRenderPipeline/Source/MovieRenderPipelineCore/Public/MoviePipelineVideoOutputBase.h
--- a/Engine/Plugins/MovieScene/MovieRenderPipeline/Source/MovieRenderPipelineCore/Public/MoviePipelineVideoOutputBase.h
+++ b/Engine/Plugins/MovieScene/MovieRenderPipeline/Source/MovieRenderPipelineCore/Public/MoviePipelineVideoOutputBase.h
@@ -14,6 +14,10 @@ namespace MovieRenderPipeline
 {
 	struct IVideoCodecWriter
 	{
+		IVideoCodecWriter()
+			: bConvertToSrgb(true)
+		{}
+
 		// The filename to actually write to disk with. This may include "de-duplication" numbers (MyFile (1).ext etc.)
 		FString FileName;
 		// The filename without de-duplication numbers, used to match up multiple incoming frames back to the same writer.
@@ -22,6 +26,16 @@ namespace MovieRenderPipeline
 		FString StableFileName;
 		FMoviePipelineFormatArgs FormatArgs;
 		bool bConvertToSrgb;
+
+		// This array contains the indexes of each shot that contributed to this writer. This is needed so that after a shot/render, when
+		// all we have is the completed blocks of audio samples, we can figure out which blocks go to which video clips. Because of how you can
+		// split videos up by either shot, or by the whole sequence, there isn't an obvious mapping, so this array will tell the audio writing
+		// system which audio blocks should be sent to the shot.
+		struct FLightweightSourceData
+		{
+			int32 SubmittedFrameCount;
+		};
+		TMap<int32, FLightweightSourceData> LightweightSourceData;
 	};
 }
 
