CIS: Fixed deprecation warnings from Clang for having a deprecated property in a struct.

diff --git a/Engine/Plugins/Media/AppleProResMedia/Source/AppleProResMedia/Public/AppleProResEncoder/AppleProResEncoder.h b/Engine/Plugins/Media/AppleProResMedia/Source/AppleProResMedia/Public/AppleProResEncoder/AppleProResEncoder.h
--- a/Engine/Plugins/Media/AppleProResMedia/Source/AppleProResMedia/Public/AppleProResEncoder/AppleProResEncoder.h
+++ b/Engine/Plugins/Media/AppleProResMedia/Source/AppleProResMedia/Public/AppleProResEncoder/AppleProResEncoder.h
@@ -56,6 +56,8 @@ enum class EAppleProResEncoderScanMode : uint8
 
 struct FAppleProResEncoderOptions
 {
+	// Disable deprecation warnings for bConvertToSrgb to make Clang happy.
+	PRAGMA_DISABLE_DEPRECATION_WARNINGS
 	FAppleProResEncoderOptions()
 		: OutputFilename()
 		, Width(0)
@@ -72,6 +74,14 @@ struct FAppleProResEncoderOptions
 		, AudioSampleRate(48000)
 		, AudioChannelCount(2)
 	{}
+	
+	FAppleProResEncoderOptions(FAppleProResEncoderOptions&) = default;
+	~FAppleProResEncoderOptions() = default;
+	FAppleProResEncoderOptions(const FAppleProResEncoderOptions&) = default;
+	FAppleProResEncoderOptions(FAppleProResEncoderOptions&&) = default;
+	FAppleProResEncoderOptions& operator=(const FAppleProResEncoderOptions&) = default;
+	FAppleProResEncoderOptions& operator=(FAppleProResEncoderOptions&&) = default;
+	PRAGMA_ENABLE_DEPRECATION_WARNINGS
 
 	/** The absolute path on disk to try and save the video file to. */
 	FString OutputFilename;
