MoviePipeline: Added a new toggle, "Include Audio", on ProRes and H.264 MRG video nodes (and equivalent settings in MRQ).

diff --git a/Engine/Plugins/Media/AppleProResMedia/Source/AppleProResMedia/Private/AppleProResEncoder/AppleProResEncoder.cpp b/Engine/Plugins/Media/AppleProResMedia/Source/AppleProResMedia/Private/AppleProResEncoder/AppleProResEncoder.cpp
--- a/Engine/Plugins/Media/AppleProResMedia/Source/AppleProResMedia/Private/AppleProResEncoder/AppleProResEncoder.cpp
+++ b/Engine/Plugins/Media/AppleProResMedia/Source/AppleProResMedia/Private/AppleProResEncoder/AppleProResEncoder.cpp
@@ -329,29 +329,34 @@ bool FAppleProResEncoder::Initialize()
 			return false;
 		}
 	}
+	
 	// Add a Audio stream
 	{
-		Status = ProResFileWriterAddTrack(FileWriter, kPRMediaType_Audio, &AudioTrackId);
-		if (Status != 0)
+		if (Options.bIncludeAudio)
 		{
-			UE_LOG(LogAppleProResMedia, Error, TEXT("Failed to add audio track to file writer at path: \"%s\"."), *Options.OutputFilename);
-			return false;
-		}
+			Status = ProResFileWriterAddTrack(FileWriter, kPRMediaType_Audio, &AudioTrackId);
+			if (Status != 0)
+			{
+				UE_LOG(LogAppleProResMedia, Error, TEXT("Failed to add audio track to file writer at path: \"%s\"."), *Options.OutputFilename);
+				return false;
+			}
 		
-		// Sound tracks must match the sample rate, not frame rate.
-		if(!ensureMsgf(Options.AudioSampleRate > 0, TEXT("AudioSampleRate must be greater than zero!")))
-		{
-			return false;
-		}
+			// Sound tracks must match the sample rate, not frame rate.
+			if(!ensureMsgf(Options.AudioSampleRate > 0, TEXT("AudioSampleRate must be greater than zero!")))
+			{
+				return false;
+			}
 		
-		PRTimeScale ProResAudioTimescale = Options.AudioSampleRate;
-		Status = ProResFileWriterSetTrackMediaTimescale(FileWriter, AudioTrackId, ProResAudioTimescale);
-		if (Status != 0)
-		{
-			UE_LOG(LogAppleProResMedia, Error, TEXT("Failed to set audio track timescale to file writer at path: \"%s\"."), *Options.OutputFilename);
-			return false;
+			PRTimeScale ProResAudioTimescale = Options.AudioSampleRate;
+			Status = ProResFileWriterSetTrackMediaTimescale(FileWriter, AudioTrackId, ProResAudioTimescale);
+			if (Status != 0)
+			{
+				UE_LOG(LogAppleProResMedia, Error, TEXT("Failed to set audio track timescale to file writer at path: \"%s\"."), *Options.OutputFilename);
+				return false;
+			}
 		}
 	}
+
 	// Add a Timecode stream
 	{
 		Status = ProResFileWriterAddTrack(FileWriter, kPRMediaType_Timecode, &TimecodeTrackId);
@@ -384,11 +389,14 @@ bool FAppleProResEncoder::Initialize()
 			return false;
 		}
 
-		bSuccess = InitializeAudioTrack();
-		if (!bSuccess)
+		if (Options.bIncludeAudio)
 		{
-			UE_LOG(LogAppleProResMedia, Error, TEXT("Failed to initialize audio track with specified settings for file writer at path: \"%s\"."), *Options.OutputFilename);
-			return false;
+			bSuccess = InitializeAudioTrack();
+			if (!bSuccess)
+			{
+				UE_LOG(LogAppleProResMedia, Error, TEXT("Failed to initialize audio track with specified settings for file writer at path: \"%s\"."), *Options.OutputFilename);
+				return false;
+			}
 		}
 	}
 
@@ -545,10 +553,13 @@ void FAppleProResEncoder::Finalize()
 			UE_LOG(LogAppleProResMedia, Error, TEXT("Failed to Mark End of Data for Timecode Track for file writer at path: \"%s\"."), *Options.OutputFilename);
 		}
 
-		Status = ProResFileWriterMarkEndOfDataForTrack(FileWriter, AudioTrackId);
-		if (Status != 0)
+		if (Options.bIncludeAudio)
 		{
-			UE_LOG(LogAppleProResMedia, Error, TEXT("Failed to Mark End of Data for Audio Track for file writer at path: \"%s\"."), *Options.OutputFilename);
+			Status = ProResFileWriterMarkEndOfDataForTrack(FileWriter, AudioTrackId);
+			if (Status != 0)
+			{
+				UE_LOG(LogAppleProResMedia, Error, TEXT("Failed to Mark End of Data for Audio Track for file writer at path: \"%s\"."), *Options.OutputFilename);
+			}
 		}
 
 		// Add one more sample onto the end 
@@ -707,6 +718,12 @@ bool FAppleProResEncoder::WriteFrame(const FImagePixelData* InPixelData)
 
 bool FAppleProResEncoder::WriteAudioSample(const TArrayView<int16>& InAudioSamples)
 {
+	// If we aren't including audio, then we just say it was a success so that callers don't get confused.
+	if (!Options.bIncludeAudio)
+	{
+		return true;
+	}
+
 	if (!ensureMsgf(bInitialized || bFinalized, TEXT("WriteFrame should not be called if not initialized or after finalize! Initialized: %d Finalized: %d"), bInitialized, bFinalized))
 	{
 		return false;
diff --git a/Engine/Plugins/Media/AppleProResMedia/Source/AppleProResMedia/Private/MovieGraphAppleProResNode.cpp b/Engine/Plugins/Media/AppleProResMedia/Source/AppleProResMedia/Private/MovieGraphAppleProResNode.cpp
--- a/Engine/Plugins/Media/AppleProResMedia/Source/AppleProResMedia/Private/MovieGraphAppleProResNode.cpp
+++ b/Engine/Plugins/Media/AppleProResMedia/Source/AppleProResMedia/Private/MovieGraphAppleProResNode.cpp
@@ -17,6 +17,7 @@
 
 UMovieGraphAppleProResNode::UMovieGraphAppleProResNode()
 	: Quality(EAppleProResEncoderCodec::ProRes_422LT)
+	, bIncludeAudio(true)
 {
 	
 }
@@ -79,6 +80,7 @@ TUniquePtr<MovieRenderGraph::IVideoCodecWriter> UMovieGraphAppleProResNode::Init
 	Options.ScanMode = EAppleProResEncoderScanMode::IM_PROGRESSIVE_SCAN; // No interlace sources
 	Options.bWriteAlpha = true;
 	Options.bDropFrameTimecode = OutputSetting->bDropFrameTimecode;
+	Options.bIncludeAudio = EvaluatedNode->bIncludeAudio;
 
 	// If OCIO is enabled, don't do additional color conversion
 	
diff --git a/Engine/Plugins/Media/AppleProResMedia/Source/AppleProResMedia/Private/MovieGraphAppleProResNode.h b/Engine/Plugins/Media/AppleProResMedia/Source/AppleProResMedia/Private/MovieGraphAppleProResNode.h
--- a/Engine/Plugins/Media/AppleProResMedia/Source/AppleProResMedia/Private/MovieGraphAppleProResNode.h
+++ b/Engine/Plugins/Media/AppleProResMedia/Source/AppleProResMedia/Private/MovieGraphAppleProResNode.h
@@ -54,6 +54,9 @@ public:
 	UPROPERTY(EditAnywhere, BlueprintReadWrite, Category = Overrides, meta = (InlineEditConditionToggle))
 	uint8 bOverride_Quality : 1;
 
+	UPROPERTY(EditAnywhere, BlueprintReadWrite, Category = Overrides, meta = (InlineEditConditionToggle))
+	uint8 bOverride_bIncludeAudio : 1;
+
 	UPROPERTY(EditAnywhere, BlueprintReadWrite, Category = Overrides, meta = (InlineEditConditionToggle))
 	uint8 bOverride_OCIOConfiguration : 1;
 
@@ -64,6 +67,10 @@ public:
 	UPROPERTY(EditAnywhere, BlueprintReadWrite, Category = "Apple ProRes", meta = (EditCondition = "bOverride_Quality"))
 	EAppleProResEncoderCodec Quality;
 
+	/** If true, audio will be included in the video file. */
+	UPROPERTY(EditAnywhere, BlueprintReadWrite, Category = "Apple ProRes", meta = (EditCondition = "bOverride_bIncludeAudio"))
+	bool bIncludeAudio;
+
 	/**
 	* OCIO configuration/transform settings.
 	*
diff --git a/Engine/Plugins/Media/AppleProResMedia/Source/AppleProResMedia/Private/MoviePipelineAppleProResOutput.cpp b/Engine/Plugins/Media/AppleProResMedia/Source/AppleProResMedia/Private/MoviePipelineAppleProResOutput.cpp
--- a/Engine/Plugins/Media/AppleProResMedia/Source/AppleProResMedia/Private/MoviePipelineAppleProResOutput.cpp
+++ b/Engine/Plugins/Media/AppleProResMedia/Source/AppleProResMedia/Private/MoviePipelineAppleProResOutput.cpp
@@ -33,6 +33,7 @@ TUniquePtr<MovieRenderPipeline::IVideoCodecWriter> UMoviePipelineAppleProResOutp
 	Options.ColorPrimaries = EAppleProResEncoderColorPrimaries::CD_HDREC709; // Force Rec 709 for now
 	Options.ScanMode = EAppleProResEncoderScanMode::IM_PROGRESSIVE_SCAN; // No interlace sources.
 	Options.bWriteAlpha = true;
+	Options.bIncludeAudio = bIncludeAudio;
 
 	TUniquePtr<FAppleProResEncoder> Encoder = MakeUnique<FAppleProResEncoder>(Options);
 	
diff --git a/Engine/Plugins/Media/AppleProResMedia/Source/AppleProResMedia/Private/MoviePipelineAppleProResOutput.h b/Engine/Plugins/Media/AppleProResMedia/Source/AppleProResMedia/Private/MoviePipelineAppleProResOutput.h
--- a/Engine/Plugins/Media/AppleProResMedia/Source/AppleProResMedia/Private/MoviePipelineAppleProResOutput.h
+++ b/Engine/Plugins/Media/AppleProResMedia/Source/AppleProResMedia/Private/MoviePipelineAppleProResOutput.h
@@ -18,6 +18,7 @@ class UMoviePipelineAppleProResOutput : public UMoviePipelineVideoOutputBase
 		: UMoviePipelineVideoOutputBase()
 		, Codec(EAppleProResEncoderCodec::ProRes_4444XQ)
 		, bDropFrameTimecode(false)
+		, bIncludeAudio(true)
 		, bOverrideMaximumEncodingThreads(false)
 		, MaxNumberOfEncodingThreads(0)
 	{
@@ -53,6 +54,10 @@ public:
 	UPROPERTY(EditAnywhere, BlueprintReadWrite, Category = "Settings")
 	bool bDropFrameTimecode;
 
+	/** If true, audio will be included in the video file. */
+	UPROPERTY(EditAnywhere, BlueprintReadWrite, Category = "Settings")
+	bool bIncludeAudio;
+
 	/** Should we override the maximum number of encoding threads? */
 	UPROPERTY(EditAnywhere, BlueprintReadWrite, Category = "Settings")
 	bool bOverrideMaximumEncodingThreads;
diff --git a/Engine/Plugins/Media/AppleProResMedia/Source/AppleProResMedia/Public/AppleProResEncoder/AppleProResEncoder.h b/Engine/Plugins/Media/AppleProResMedia/Source/AppleProResMedia/Public/AppleProResEncoder/AppleProResEncoder.h
--- a/Engine/Plugins/Media/AppleProResMedia/Source/AppleProResMedia/Public/AppleProResEncoder/AppleProResEncoder.h
+++ b/Engine/Plugins/Media/AppleProResMedia/Source/AppleProResMedia/Public/AppleProResEncoder/AppleProResEncoder.h
@@ -71,6 +71,7 @@ struct FAppleProResEncoderOptions
 		, bWriteAlpha(false)
 		, bDropFrameTimecode(false)
 		, FrameNumberOffset(0)
+		, bIncludeAudio(true)
 		, AudioSampleRate(48000)
 		, AudioChannelCount(2)
 	{}
@@ -120,6 +121,9 @@ struct FAppleProResEncoderOptions
 	/** The number of frames to offset the Timecode track by. */
 	int32 FrameNumberOffset;
 
+	/** If false, then audio tracks will not be written and calls to WriteAudioSample will be ignored. */
+	bool bIncludeAudio;
+
 	/** The sample rate of the incoming audio data. Only tested with 48'000. */
 	uint32 AudioSampleRate;
 
diff --git a/Engine/Plugins/MovieScene/MovieRenderPipeline/Source/MovieRenderPipelineMP4Encoder/Private/Graph/MovieGraphMP4EncoderNode.cpp b/Engine/Plugins/MovieScene/MovieRenderPipeline/Source/MovieRenderPipelineMP4Encoder/Private/Graph/MovieGraphMP4EncoderNode.cpp
--- a/Engine/Plugins/MovieScene/MovieRenderPipeline/Source/MovieRenderPipelineMP4Encoder/Private/Graph/MovieGraphMP4EncoderNode.cpp
+++ b/Engine/Plugins/MovieScene/MovieRenderPipeline/Source/MovieRenderPipelineMP4Encoder/Private/Graph/MovieGraphMP4EncoderNode.cpp
@@ -79,6 +79,7 @@ TUniquePtr<MovieRenderGraph::IVideoCodecWriter> UMovieGraphMP4EncoderNode::Initi
 	Options.EncodingProfile = EvaluatedNode->EncodingProfile;
 	Options.EncodingLevel = EvaluatedNode->EncodingLevel;
 	Options.EncodingRateControl = EvaluatedNode->EncodingRateControl;
+	Options.bIncludeAudio = EvaluatedNode->bIncludeAudio;
 	
 	TUniquePtr<FMP4CodecWriter> NewWriter = MakeUnique<FMP4CodecWriter>();
 	NewWriter->Writer = MakeUnique<FMoviePipelineMP4Encoder>(Options);
diff --git a/Engine/Plugins/MovieScene/MovieRenderPipeline/Source/MovieRenderPipelineMP4Encoder/Private/MoviePipelineMP4EncoderCommon.h b/Engine/Plugins/MovieScene/MovieRenderPipeline/Source/MovieRenderPipelineMP4Encoder/Private/MoviePipelineMP4EncoderCommon.h
--- a/Engine/Plugins/MovieScene/MovieRenderPipeline/Source/MovieRenderPipelineMP4Encoder/Private/MoviePipelineMP4EncoderCommon.h
+++ b/Engine/Plugins/MovieScene/MovieRenderPipeline/Source/MovieRenderPipelineMP4Encoder/Private/MoviePipelineMP4EncoderCommon.h
@@ -71,6 +71,7 @@ struct FMoviePipelineMP4EncoderOptions
 		, Width(0)
 		, Height(0)
 		, FrameRate(30, 1)
+		, bIncludeAudio(true)
 		, AudioChannelCount(2)
 		, AudioSampleRate(48000)
 		, AudioAverageBitRate(24000) // 192 kbps (bits)
@@ -94,6 +95,9 @@ struct FMoviePipelineMP4EncoderOptions
 
 	/** Frame Rate of the output video. */
 	FFrameRate FrameRate;
+	
+	/** If false, then audio tracks will not be written and calls to WriteAudioSample will be ignored. */
+	bool bIncludeAudio;
 
 	/** Number of audio channels for the audio track. */
 	uint32 AudioChannelCount;
diff --git a/Engine/Plugins/MovieScene/MovieRenderPipeline/Source/MovieRenderPipelineMP4Encoder/Private/MoviePipelineMP4EncoderOutput.cpp b/Engine/Plugins/MovieScene/MovieRenderPipeline/Source/MovieRenderPipelineMP4Encoder/Private/MoviePipelineMP4EncoderOutput.cpp
--- a/Engine/Plugins/MovieScene/MovieRenderPipeline/Source/MovieRenderPipelineMP4Encoder/Private/MoviePipelineMP4EncoderOutput.cpp
+++ b/Engine/Plugins/MovieScene/MovieRenderPipeline/Source/MovieRenderPipelineMP4Encoder/Private/MoviePipelineMP4EncoderOutput.cpp
@@ -33,6 +33,7 @@ TUniquePtr<MovieRenderPipeline::IVideoCodecWriter> UMoviePipelineMP4EncoderOutpu
 	Options.EncodingProfile = EncodingProfile;
 	Options.EncodingLevel = EncodingLevel;
 	Options.EncodingRateControl = EncodingRateControl;
+	Options.bIncludeAudio = bIncludeAudio;
 	
 	TUniquePtr<FMoviePipelineMP4Encoder> Writer = MakeUnique<FMoviePipelineMP4Encoder>(Options);
 	
diff --git a/Engine/Plugins/MovieScene/MovieRenderPipeline/Source/MovieRenderPipelineMP4Encoder/Private/Windows/MoviePipelineMP4Encoder.cpp b/Engine/Plugins/MovieScene/MovieRenderPipeline/Source/MovieRenderPipelineMP4Encoder/Private/Windows/MoviePipelineMP4Encoder.cpp
--- a/Engine/Plugins/MovieScene/MovieRenderPipeline/Source/MovieRenderPipelineMP4Encoder/Private/Windows/MoviePipelineMP4Encoder.cpp
+++ b/Engine/Plugins/MovieScene/MovieRenderPipeline/Source/MovieRenderPipelineMP4Encoder/Private/Windows/MoviePipelineMP4Encoder.cpp
@@ -208,6 +208,12 @@ bool FMoviePipelineMP4Encoder::WriteAudioSample(const TArrayView<int16>& InAudio
 {
 	TRACE_CPUPROFILER_EVENT_SCOPE(MRQ_WMFSink_WriteAudioSample);
 
+	// If we aren't including audio, then we just say it was a success so that callers don't get confused.
+	if (!Options.bIncludeAudio)
+	{
+		return true;
+	}
+
 	// If we were not initialized, we can't write the sample.
 	if (!ensureMsgf(bInitialized || bFinalized, TEXT("WriteAudioSample should not be called if not initialized or after finalize! Initialized: %d Finalized: %d"), bInitialized, bFinalized))
 	{
@@ -748,20 +754,23 @@ bool FMoviePipelineMP4Encoder::InitializeEncoder()
 		Result = SinkWriter->AddStream(VideoMediaTypeOut, (DWORD*)&VideoStreamIndex);
 	}
 
-	if (SUCCEEDED(Result))
+	if (Options.bIncludeAudio)
 	{
-		Result = CreateAudioMediaOutputStream(&AudioMediaTypeOut, Options);
-	}
+		if (SUCCEEDED(Result))
+		{
+			Result = CreateAudioMediaOutputStream(&AudioMediaTypeOut, Options);
+		}
 
-	// Audio Stream
-	if (SUCCEEDED(Result) && SinkWriter)
-	{
-		Result = SinkWriter->AddStream(AudioMediaTypeOut, (DWORD*)&AudioStreamIndex);
-	}
+		// Audio Stream
+		if (SUCCEEDED(Result) && SinkWriter)
+		{
+			Result = SinkWriter->AddStream(AudioMediaTypeOut, (DWORD*)&AudioStreamIndex);
+		}
 
-	if (SUCCEEDED(Result))
-	{
-		Result = CreateAudioMediaTypeIn(&AudioMediaTypeIn, Options);
+		if (SUCCEEDED(Result))
+		{
+			Result = CreateAudioMediaTypeIn(&AudioMediaTypeIn, Options);
+		}
 	}
 
 	if (SUCCEEDED(Result))
@@ -784,9 +793,12 @@ bool FMoviePipelineMP4Encoder::InitializeEncoder()
 		}
 	}
 
-	if (SUCCEEDED(Result) && SinkWriter)
+	if (Options.bIncludeAudio)
 	{
-		Result = SinkWriter->SetInputMediaType(AudioStreamIndex, AudioMediaTypeIn, NULL);
+		if (SUCCEEDED(Result) && SinkWriter)
+		{
+			Result = SinkWriter->SetInputMediaType(AudioStreamIndex, AudioMediaTypeIn, NULL);
+		}
 	}
 
 	// Note: This has no practical use to the user. However, for diagnostic purposes, `codecApi` can be inspected in the debugger to determine which
diff --git a/Engine/Plugins/MovieScene/MovieRenderPipeline/Source/MovieRenderPipelineMP4Encoder/Public/Graph/MovieGraphMP4EncoderNode.h b/Engine/Plugins/MovieScene/MovieRenderPipeline/Source/MovieRenderPipelineMP4Encoder/Public/Graph/MovieGraphMP4EncoderNode.h
--- a/Engine/Plugins/MovieScene/MovieRenderPipeline/Source/MovieRenderPipelineMP4Encoder/Public/Graph/MovieGraphMP4EncoderNode.h
+++ b/Engine/Plugins/MovieScene/MovieRenderPipeline/Source/MovieRenderPipelineMP4Encoder/Public/Graph/MovieGraphMP4EncoderNode.h
@@ -57,6 +57,9 @@ public:
 	UPROPERTY(EditAnywhere, BlueprintReadWrite, Category = Overrides, meta = (InlineEditConditionToggle))
 	uint8 bOverride_AverageBitrateInMbps : 1;
 	
+	UPROPERTY(EditAnywhere, BlueprintReadWrite, Category = Overrides, meta = (InlineEditConditionToggle))
+	uint8 bOverride_bIncludeAudio : 1;
+	
 	UPROPERTY(EditAnywhere, BlueprintReadWrite, Category = Overrides, meta = (InlineEditConditionToggle))
 	uint8 bOverride_MaxBitrateInMbps : 1;
 	
@@ -82,6 +85,10 @@ public:
 	*/
 	UPROPERTY(BlueprintReadWrite, EditAnywhere, Category = "H.264", meta=(EditCondition="bOverride_EncodingRateControl"))
 	EMoviePipelineMP4EncodeRateControlMode EncodingRateControl = EMoviePipelineMP4EncodeRateControlMode::Quality;
+
+	/** If true, audio will be included in the video file. */
+	UPROPERTY(EditAnywhere, BlueprintReadWrite, Category = "H.264", meta = (EditCondition = "bOverride_bIncludeAudio"))
+	bool bIncludeAudio;
 	
 	/**
 	* What is the average bitrate the encoder should target per second? Value is in Megabits per Second,
diff --git a/Engine/Plugins/MovieScene/MovieRenderPipeline/Source/MovieRenderPipelineMP4Encoder/Public/MoviePipelineMP4EncoderOutput.h b/Engine/Plugins/MovieScene/MovieRenderPipeline/Source/MovieRenderPipelineMP4Encoder/Public/MoviePipelineMP4EncoderOutput.h
--- a/Engine/Plugins/MovieScene/MovieRenderPipeline/Source/MovieRenderPipelineMP4Encoder/Public/MoviePipelineMP4EncoderOutput.h
+++ b/Engine/Plugins/MovieScene/MovieRenderPipeline/Source/MovieRenderPipelineMP4Encoder/Public/MoviePipelineMP4EncoderOutput.h
@@ -26,7 +26,7 @@ protected:
 	virtual void BeginFinalize_EncodeThread(MovieRenderPipeline::IVideoCodecWriter* InWriter);
 	virtual void Finalize_EncodeThread(MovieRenderPipeline::IVideoCodecWriter* InWriter);
 	virtual const TCHAR* GetFilenameExtension() const override { return TEXT("mp4"); }
-	virtual bool IsAudioSupported() const override { return false; }
+	virtual bool IsAudioSupported() const override { return true; }
 	// ~UMoviePipelineVideoOutputBase Interface
 	
 	// UMoviePipelineSetting Interface
@@ -48,6 +48,10 @@ public:
 	UPROPERTY(BlueprintReadWrite, EditAnywhere, Category = "Settings")
 	EMoviePipelineMP4EncodeRateControlMode EncodingRateControl = EMoviePipelineMP4EncodeRateControlMode::Quality;
 	
+	/** If true, audio will be included in the video file. */
+	UPROPERTY(EditAnywhere, BlueprintReadWrite, Category = "Settings", meta = (EditCondition = "bOverride_bIncludeAudio"))
+	bool bIncludeAudio = true;
+	
 	/**
 	* What is the average bitrate the encoder should target per second? Value is in Megabits per Second,
 	* so a value of 8 will become 8192Kbps (kilobits per second). Higher resolutions and higher framerates
