fix copyright

#rnx

diff --git a/Engine/Plugins/Media/AppleProResMedia/Source/AppleProResMedia/Private/Mac/ProResLibMac.cpp b/Engine/Plugins/Media/AppleProResMedia/Source/AppleProResMedia/Private/Mac/ProResLibMac.cpp
--- a/Engine/Plugins/Media/AppleProResMedia/Source/AppleProResMedia/Private/Mac/ProResLibMac.cpp
+++ b/Engine/Plugins/Media/AppleProResMedia/Source/AppleProResMedia/Private/Mac/ProResLibMac.cpp
@@ -1,9 +1,4 @@
 //	Copyright Epic Games, Inc. All Rights Reserved.
-//
-//  Mac implementation of ProResToolbox and ProResLib, using AVFoundation and VideoToolbox
-
-#import <AVFoundation/AVFoundation.h>
-#import <VideoToolbox/VideoToolbox.h>
 
 #include "ProResEncoder.h"
 #include "ProResProperties.h"
@@ -14,6 +9,11 @@
 #include "ProResTypes.h"
 #include "ProResToolboxExtension.h"
 
+#import <AVFoundation/AVFoundation.h>
+#import <VideoToolbox/VideoToolbox.h>
+
+//  Mac implementation of ProResToolbox and ProResLib, using AVFoundation and VideoToolbox
+
 struct PREncoder
 {
 	VTCompressionSessionRef session = NULL;
diff --git a/Engine/Plugins/Media/AppleProResMedia/Source/AppleProResMedia/Private/ProResToolboxExtension.h b/Engine/Plugins/Media/AppleProResMedia/Source/AppleProResMedia/Private/ProResToolboxExtension.h
--- a/Engine/Plugins/Media/AppleProResMedia/Source/AppleProResMedia/Private/ProResToolboxExtension.h
+++ b/Engine/Plugins/Media/AppleProResMedia/Source/AppleProResMedia/Private/ProResToolboxExtension.h
@@ -1,8 +1,5 @@
 //	Copyright Epic Games, Inc. All Rights Reserved.
 
-#ifndef ProResToolboxExtension_h
-#define ProResToolboxExtension_h
-
 #include "ProResEncoder.h"
 #include "ProResFileWriter.h"
 
@@ -15,5 +12,3 @@ PR_EXPORT PRStatus ProResFileWriterAddEncoderSampleBufferToTrack(
 														Only usable with video track. */
 	PRTime				timestamp);					/*! @param timestamp
 														OutputPresentationTimeStamp. Use the timestamp in PRSampleTimingInfo.timeStamp. */
-
-#endif /* ProResToolboxExtension_h */
