fix copyright

#rnx

diff --git a/Engine/Plugins/Media/AppleProResMedia/Source/AppleProResMedia/Private/Mac/ProResLibMac.cpp b/Engine/Plugins/Media/AppleProResMedia/Source/AppleProResMedia/Private/Mac/ProResLibMac.cpp
--- a/Engine/Plugins/Media/AppleProResMedia/Source/AppleProResMedia/Private/Mac/ProResLibMac.cpp
+++ b/Engine/Plugins/Media/AppleProResMedia/Source/AppleProResMedia/Private/Mac/ProResLibMac.cpp
@@ -1,4 +1,8 @@
-//	Copyright Epic Games, Inc. All Rights Reserved.
+// Copyright Epic Games, Inc. All Rights Reserved.
+// Mac implementation of ProResToolbox and ProResLib, using AVFoundation and VideoToolbox
+
+#import <AVFoundation/AVFoundation.h>
+#import <VideoToolbox/VideoToolbox.h>
 
 #include "ProResEncoder.h"
 #include "ProResProperties.h"
@@ -9,11 +13,6 @@
 #include "ProResTypes.h"
 #include "ProResToolboxExtension.h"
 
-#import <AVFoundation/AVFoundation.h>
-#import <VideoToolbox/VideoToolbox.h>
-
-//  Mac implementation of ProResToolbox and ProResLib, using AVFoundation and VideoToolbox
-
 struct PREncoder
 {
 	VTCompressionSessionRef session = NULL;
diff --git a/Engine/Plugins/Media/AppleProResMedia/Source/AppleProResMedia/Private/ProResToolboxExtension.h b/Engine/Plugins/Media/AppleProResMedia/Source/AppleProResMedia/Private/ProResToolboxExtension.h
--- a/Engine/Plugins/Media/AppleProResMedia/Source/AppleProResMedia/Private/ProResToolboxExtension.h
+++ b/Engine/Plugins/Media/AppleProResMedia/Source/AppleProResMedia/Private/ProResToolboxExtension.h
@@ -1,4 +1,7 @@
-//	Copyright Epic Games, Inc. All Rights Reserved.
+// Copyright Epic Games, Inc. All Rights Reserved.
+
+#ifndef ProResToolboxExtension_h
+#define ProResToolboxExtension_h
 
 #include "ProResEncoder.h"
 #include "ProResFileWriter.h"
@@ -12,3 +15,5 @@ PR_EXPORT PRStatus ProResFileWriterAddEncoderSampleBufferToTrack(
 														Only usable with video track. */
 	PRTime				timestamp);					/*! @param timestamp
 														OutputPresentationTimeStamp. Use the timestamp in PRSampleTimingInfo.timeStamp. */
+
+#endif /* ProResToolboxExtension_h */
