Mac ProRes fixes:
Adding color primaries, matrix coefficient, framerate to frameHeader data
Fix blank frame at 0
Fix memory leak
Fix incorrect timecode info

#rnx

diff --git a/Engine/Plugins/Media/AppleProResMedia/Source/AppleProResMedia/Private/AppleProResEncoder/AppleProResEncoder.cpp b/Engine/Plugins/Media/AppleProResMedia/Source/AppleProResMedia/Private/AppleProResEncoder/AppleProResEncoder.cpp
--- a/Engine/Plugins/Media/AppleProResMedia/Source/AppleProResMedia/Private/AppleProResEncoder/AppleProResEncoder.cpp
+++ b/Engine/Plugins/Media/AppleProResMedia/Source/AppleProResMedia/Private/AppleProResEncoder/AppleProResEncoder.cpp
@@ -672,7 +672,7 @@ bool FAppleProResEncoder::WriteFrame(const FImagePixelData* InPixelData)
 		FileWriter,
 		VideoTrackId,
 		Encoder,
-		CurrentVideoTime);
+		SampleTimingInfo.timeStamp);
 #else
 	// Add the sample to the Video Track
 	PRStatus Status = ProResFileWriterAddSampleBufferToTrack(
diff --git a/Engine/Plugins/Media/AppleProResMedia/Source/AppleProResMedia/Private/Mac/ProResLibMac.cpp b/Engine/Plugins/Media/AppleProResMedia/Source/AppleProResMedia/Private/Mac/ProResLibMac.cpp
--- a/Engine/Plugins/Media/AppleProResMedia/Source/AppleProResMedia/Private/Mac/ProResLibMac.cpp
+++ b/Engine/Plugins/Media/AppleProResMedia/Source/AppleProResMedia/Private/Mac/ProResLibMac.cpp
@@ -218,18 +218,37 @@ int PREncodeFrame(PREncoderRef encoder,
 			{
 				CFDictionaryAddValue(encoder->sessionProperties,
 									 kVTCompressionPropertyKey_FieldDetail,
-									 kCVImageBufferFieldDetailTemporalTopFirst);
+									 kCMFormatDescriptionFieldDetail_TemporalTopFirst);
 			}
 			else if (encodingParams->interlaceMode == kPRInterlacedTopFieldFirst)
 			{
 				CFDictionaryAddValue(encoder->sessionProperties,
 									 kVTCompressionPropertyKey_FieldDetail,
-									 kCVImageBufferFieldDetailTemporalBottomFirst);
+									 kCMFormatDescriptionFieldDetail_TemporalBottomFirst);
 			}
 		}
 		CFDictionaryAddValue(encoder->sessionProperties,
 							 kVTCompressionPropertyKey_PreserveAlphaChannel,
 							 encodingParams->preserveAlpha ? kCFBooleanTrue : kCFBooleanFalse);
+		// default color primaries
+		if (!CFDictionaryContainsKey(encoder->sessionProperties, kVTCompressionPropertyKey_ColorPrimaries))
+		{
+			CFDictionaryAddValue(encoder->sessionProperties,
+								 kVTCompressionPropertyKey_ColorPrimaries,
+								 kCMFormatDescriptionColorPrimaries_ITU_R_709_2);
+		}
+		if (!CFDictionaryContainsKey(encoder->sessionProperties, kVTCompressionPropertyKey_TransferFunction))
+		{
+			CFDictionaryAddValue(encoder->sessionProperties,
+								 kVTCompressionPropertyKey_TransferFunction,
+								 kCMFormatDescriptionTransferFunction_ITU_R_709_2);
+		}
+		if (!CFDictionaryContainsKey(encoder->sessionProperties, kVTCompressionPropertyKey_YCbCrMatrix))
+		{
+			CFDictionaryAddValue(encoder->sessionProperties,
+								 kVTCompressionPropertyKey_YCbCrMatrix,
+								 kCMFormatDescriptionYCbCrMatrix_ITU_R_709_2);
+		}
 		VTSessionSetProperties(encoder->session, encoder->sessionProperties);
 	}
 	
@@ -257,7 +276,7 @@ int PREncodeFrame(PREncoderRef encoder,
 				NSLog(@"Will overwrite encoder->sampleBufferOut! You should call ProResFileWriterAddEncoderSampleBufferToTrack before encoding the next frame!");
 				CFRelease(encoder->sampleBufferOut);
 			}
-			encoder->sampleBufferOut = nullptr;
+			encoder->sampleBufferOut = NULL;
 			VTEncodeInfoFlags flags;
 			status = VTCompressionSessionEncodeFrameWithOutputHandler(encoder->session,
 																	  pixelBuffer,
@@ -342,6 +361,7 @@ PR_EXPORT void PRRelease(PRTypeRef type)
 	}
 }
 
+// matching kPRFrameRate in ProResEncoder.h
 constexpr float kFrameRateFromPRFrameRateEnum[12] = {
 	0.0f,
 	24/1.001f,
@@ -354,7 +374,8 @@ constexpr float kFrameRateFromPRFrameRateEnum[12] = {
 	60.0f,
 	100.0f,
 	120/1.001f,
-	120.0f};
+	120.0f
+};
 
 int
 PRSetEncoderProperty(
@@ -364,37 +385,119 @@ PRSetEncoderProperty(
    const void*  propValueAddress)
 {
 	int propEnum = *(const int*)propValueAddress;
+	CFStringRef key;
+	CFTypeRef value = NULL;
+	bool bNeedRelease = false;
 	switch (propID)
 	{
 		case kPRPropertyID_FrameHeaderFrameRate:
-			CFDictionaryAddValue(encoder->sessionProperties,
-								 kVTCompressionPropertyKey_ExpectedFrameRate,
-								 CFNumberCreate(kCFAllocatorDefault, kCFNumberFloat32Type, &kFrameRateFromPRFrameRateEnum[propEnum]));
+			key = kVTCompressionPropertyKey_ExpectedFrameRate;
+			value = CFNumberCreate(kCFAllocatorDefault, kCFNumberFloat32Type, &kFrameRateFromPRFrameRateEnum[propEnum]);
+			bNeedRelease = true;
 			break;
 		case kPRPropertyID_FrameHeaderAspectRatio:
-//			CFDictionaryAddValue(encoder->sessionProperties,
-//								 kVTCompressionPropertyKey_PixelAspectRatio,
-//								 CFNumberCreate(kCFAllocatorDefault, kCFNumberIntType, propValueAddress));	// todo: translate enum to CFDict
-//			break;
+		{
+			key = kVTCompressionPropertyKey_PixelAspectRatio;
+			CFMutableDictionaryRef valueDict = CFDictionaryCreateMutable(kCFAllocatorDefault,
+																		 2,
+																		 &kCFTypeDictionaryKeyCallBacks,
+																		 &kCFTypeDictionaryValueCallBacks);
+			int width = -1;
+			int height = -1;
+			switch (propEnum) {
+				case kPRAspectRatio_Unspecified:
+				case kPRAspectRatio_SquarePixel:
+				default:
+					width = 1;
+					height = 1;
+					break;
+				case kPRAspectRatio_16x9:
+					width = 16;
+					height = 9;
+					break;
+				case kPRAspectRatio_4x3:
+					width = 4;
+					height = 3;
+					break;
+			}
+			CFNumberRef widthRef = CFNumberCreate(kCFAllocatorDefault, kCFNumberIntType, &width);
+			CFDictionaryAddValue(valueDict, kCMFormatDescriptionKey_PixelAspectRatioHorizontalSpacing, widthRef);
+			CFRelease(widthRef);
+			CFNumberRef heightRef = CFNumberCreate(kCFAllocatorDefault, kCFNumberIntType, &height);
+			CFDictionaryAddValue(valueDict, kCMFormatDescriptionKey_PixelAspectRatioVerticalSpacing, heightRef);
+			CFRelease(heightRef);
+			
+			value = valueDict;
+			bNeedRelease = true;
+			
+			break;
+		}
 		case kPRPropertyID_FrameHeaderColorPrimaries:
-//			CFDictionaryAddValue(encoder->sessionProperties,
-//								 kVTCompressionPropertyKey_ColorPrimaries,
-//								 CFNumberCreate(kCFAllocatorDefault, kCFNumberIntType, propValueAddress));	// todo: should be CFString
-//			break;
+			key = kVTCompressionPropertyKey_ColorPrimaries;
+			switch (propEnum)
+			{
+				case kPRColorPrimaries_ITU_R_709:
+				case kPRColorPrimaries_Unspecified:
+				default:
+					value = kCMFormatDescriptionColorPrimaries_ITU_R_709_2;
+					break;
+				case kPRColorPrimaries_EBU_3213:
+					value = kCMFormatDescriptionColorPrimaries_EBU_3213;
+					break;
+				case kPRColorPrimaries_SMPTE_C:
+					value = kCMFormatDescriptionColorPrimaries_SMPTE_C;
+					break;
+				case kPRColorPrimaries_ITU_R_2020:
+					value = kCMFormatDescriptionColorPrimaries_ITU_R_2020;
+					break;
+				case kPRColorPrimaries_DCI_P3:
+					value = kCMFormatDescriptionColorPrimaries_DCI_P3;
+					break;
+				case kPRColorPrimaries_P3_D65:
+					value = kCMFormatDescriptionColorPrimaries_P3_D65;
+					break;
+			}
+			break;
 		case kPRPropertyID_FrameHeaderTransferCharacteristic:
-//			CFDictionaryAddValue(encoder->sessionProperties,
-//								 kVTCompressionPropertyKey_TransferFunction,
-//								 CFNumberCreate(kCFAllocatorDefault, kCFNumberIntType, propValueAddress));	// todo: should be CFString
-//			break;
+			key = kVTCompressionPropertyKey_TransferFunction;
+			switch (propEnum)
+			{
+				case kPRTransferCharacteristic_ITU_R_709:
+				case kPRTransferCharacteristic_Unspecified:
+				default:
+					value = kCMFormatDescriptionTransferFunction_ITU_R_709_2;
+					break;
+				case kPRTransferCharacteristic_ST_2084:
+					value = kCMFormatDescriptionTransferFunction_SMPTE_ST_2084_PQ;
+					break;
+				case kPRTransferCharacteristic_HLG:
+					value = kCMFormatDescriptionTransferFunction_ITU_R_2100_HLG;
+					break;
+			}
+			break;
 		case kPRPropertyID_FrameHeaderMatrixCoefficients:
-//			CFDictionaryAddValue(encoder->sessionProperties,
-//								 kVTCompressionPropertyKey_YCbCrMatrix,
-//								 CFNumberCreate(kCFAllocatorDefault, kCFNumberIntType, propValueAddress));	// todo: should be CFString
-//			break;
+			key = kVTCompressionPropertyKey_YCbCrMatrix;
+			switch (propEnum)
+			{
+				case kPRMatrixCoefficients_ITU_R_709:
+				case kPRMatrixCoefficients_Unspecified:
+				default:
+					value = kCMFormatDescriptionYCbCrMatrix_ITU_R_709_2;
+					break;
+				case kPRMatrixCoefficients_ITU_R_2020:
+					value = kCMFormatDescriptionYCbCrMatrix_ITU_R_2020;
+					break;
+				case kPRMatrixCoefficients_ITU_R_601:
+					value = kCMFormatDescriptionYCbCrMatrix_ITU_R_601_4;
+					break;
+			}
+			break;
 		default:
 			NSLog(@"Unsupported property!");
 			return -2195;
 	}
+	CFDictionaryAddValue(encoder->sessionProperties, key, value);
+	
 	return noErr;
 }
 
@@ -428,10 +531,12 @@ PR_EXPORT PRStatus ProResFileWriterSetTrackMediaTimescale(
 	PRPersistentTrackID writerTrackID,
 	PRTimeScale timescale )
 {
-	// todo: doesn't work on audio???
-	if ([writer->innerRef.inputs objectAtIndex:writerTrackID].mediaType != AVMediaTypeAudio)
+	AVAssetWriterInput* writerInput = [writer->innerRef.inputs objectAtIndex:writerTrackID];
+	// AVFoundation does not allow setting mediaTimeScale with media type AVMediaTypeAudio
+	// see AVAssetWriterInput.h:439
+	if (writerInput.mediaType != AVMediaTypeAudio)
 	{
-		[writer->innerRef.inputs objectAtIndex:writerTrackID].mediaTimeScale = timescale;
+		writerInput.mediaTimeScale = timescale;
 	}
 	return writer->innerRef.error != NULL;
 }
@@ -441,7 +546,7 @@ PR_EXPORT PRStatus ProResFileWriterSetTrackPreferredChunkSize(
 	PRPersistentTrackID writerTrackID,
 	int32_t preferredChunkSize )
 {
-	// Ignore and use default, otherwise will result in some serious padding and bloat up files sometimes 8x
+	// Ignore and use default, otherwise will result in some serious padding and bloat up file size sometimes 8x
 	//[writer->innerRef.inputs objectAtIndex:writerTrackID].preferredMediaChunkAlignment = preferredChunkSize;
 	return writer->innerRef.error != NULL;
 }
@@ -465,7 +570,7 @@ PR_EXPORT PRStatus ProResFileWriterBeginSession(
 	PRTime sessionStartTime )
 {
 	[writer->innerRef startWriting];
-	[writer->innerRef startSessionAtSourceTime:kCMTimeZero];
+	[writer->innerRef startSessionAtSourceTime:CMTimeFromPRTime(sessionStartTime)];
 	return writer->innerRef.error != NULL;
 }
 
@@ -492,16 +597,17 @@ PR_EXPORT PRStatus ProResFileWriterAddSampleBufferToTrack(
 	{
 		CMBlockBufferRef blockBufferRef;
 		int status = CMBlockBufferCreateWithMemoryBlock(kCFAllocatorDefault,
-														dataBuffer,
-														dataBufferLength,
+														NULL,
+														sizeof(int32_t),
 														kCFAllocatorDefault,
-														nullptr,
-														0,
-														dataBufferLength,
+														NULL,
 														0,
+														sizeof(int32_t),
+														kCMBlockBufferAssureMemoryNowFlag,
 														&blockBufferRef);
-		if (status == noErr && blockBufferRef)
+		if (status == kCMBlockBufferNoErr && blockBufferRef)
 		{
+			status = CMBlockBufferReplaceDataBytes(dataBuffer, blockBufferRef, 0, sizeof(int32_t));
 			CMSampleTimingInfo sampleTimingArrayCM[numSampleTimingEntries];
 			for (int i = 0;	i < numSampleTimingEntries; i++)
 			{
@@ -522,7 +628,7 @@ PR_EXPORT PRStatus ProResFileWriterAddSampleBufferToTrack(
 				[writerInput appendSampleBuffer:sampleBufferRef];
 				CFRelease(sampleBufferRef);
 			}
-			//CFRelease(blockBufferRef);	// todo: memory leak? crashing otherwise
+			CFRelease(blockBufferRef);
 		}
 	}
 	else
@@ -547,7 +653,7 @@ PR_EXPORT PRStatus ProResFileWriterAddEncoderSampleBufferToTrack(
 			CMSampleBufferSetOutputPresentationTimeStamp(encoder->sampleBufferOut, CMTimeFromPRTime(timestamp));
 			[writerInput appendSampleBuffer:encoder->sampleBufferOut];
 			CFRelease(encoder->sampleBufferOut);
-			encoder->sampleBufferOut = nullptr;
+			encoder->sampleBufferOut = NULL;
 		}
 		else
 		{
@@ -575,8 +681,8 @@ PR_EXPORT PRStatus ProResFileWriterAddAudioSampleBufferToTrack(
 	int status = CMBlockBufferCreateWithMemoryBlock(kCFAllocatorDefault,
 													dataBuffer,
 													dataBufferLength,
-													kCFAllocatorDefault,
-													nullptr,
+													kCFAllocatorNull,
+													NULL,
 													0,
 													dataBufferLength,
 													0,
@@ -589,7 +695,7 @@ PR_EXPORT PRStatus ProResFileWriterAddAudioSampleBufferToTrack(
 																	  formatDescription->innerRef,
 																	  numSamples,
 																	  CMTimeFromPRTime(timeStamp),
-																	  nullptr,
+																	  NULL,
 																	  &sampleBufferRef);
 		if (status == noErr && sampleBufferRef)
 		{
@@ -597,7 +703,7 @@ PR_EXPORT PRStatus ProResFileWriterAddAudioSampleBufferToTrack(
 			[writerInput appendSampleBuffer:sampleBufferRef];
 			CFRelease(sampleBufferRef);
 		}
-		//CFRelease(blockBufferRef);	// todo: memory leak? crashing otherwise
+		CFRelease(blockBufferRef);
 	}
 	return status;
 }
@@ -639,11 +745,11 @@ PR_EXPORT PRStatus ProResTimecodeFormatDescriptionCreate(
 {
 	CMFormatDescriptionRef ref;
 	int status = CMTimeCodeFormatDescriptionCreate(kCFAllocatorDefault,
-												   kCMTimeCodeFormatType_TimeCode64,	// todo: is this right?
+												   kCMTimeCodeFormatType_TimeCode32,
 												   CMTimeFromPRTime(frameDuration),
 												   frameQuanta,
 												   tcFlags,
-												   nullptr,
+												   NULL,
 												   &ref);
 	if (status == noErr)
 	{
@@ -684,8 +790,8 @@ PR_EXPORT PRStatus ProResAudioFormatDescriptionCreate(
 												layoutSize,
 												layout,
 												0,
-												nullptr,
-												nullptr,
+												NULL,
+												NULL,
 												&ref);
 	if (status == noErr)
 	{
