[AppleProResMedia] Editor load time improvement: Implementing on demand loading of the ProResToolbox only when the encoder is used.

There is a known issue with this dll on windows. If a debugger is attached to the process, it will take between 20 to 40 seconds to load (cause still under investigation).

By making it load on demand, we are at least mitigating the issue for everyone that are not using the encoder.

diff --git a/Engine/Plugins/Media/AppleProResMedia/Source/AppleProResMedia/Private/AppleProResEncoder/AppleProResEncoder.cpp b/Engine/Plugins/Media/AppleProResMedia/Source/AppleProResMedia/Private/AppleProResEncoder/AppleProResEncoder.cpp
--- a/Engine/Plugins/Media/AppleProResMedia/Source/AppleProResMedia/Private/AppleProResEncoder/AppleProResEncoder.cpp
+++ b/Engine/Plugins/Media/AppleProResMedia/Source/AppleProResMedia/Private/AppleProResEncoder/AppleProResEncoder.cpp
@@ -1,8 +1,9 @@
 // Copyright Epic Games, Inc. All Rights Reserved.
 
 #include "AppleProResEncoder/AppleProResEncoder.h"
-#include "HAL/FileManager.h"
 #include "AppleProResMediaModule.h"
+#include "HAL/FileManager.h"
+#include "Modules/ModuleManager.h"
 #include "ProResToolboxExtension.h"
 
 namespace AppleProRes
@@ -276,6 +277,16 @@ FAppleProResEncoder::~FAppleProResEncoder()
 
 bool FAppleProResEncoder::Initialize()
 {
+	// Manually loading the toolbox library, on demand only.
+	if (IAppleProResMediaModule* AppleProResMediaModule = FModuleManager::GetModulePtr<IAppleProResMediaModule>("AppleProResMedia"))
+	{
+		if (!AppleProResMediaModule->LoadProResToolbox())
+		{
+			UE_LOG(LogAppleProResMedia, Error, TEXT("Failed to initialize ProRes Encoder."));
+			return false;	
+		}
+	}
+	
 	// Attempt to create the Encoder
 	Encoder = PROpenEncoder(FMath::Max(Options.MaxNumberOfEncodingThreads, 0), /*ThreadStartupCallback*/ nullptr);
 	if (!Encoder)
diff --git a/Engine/Plugins/Media/AppleProResMedia/Source/AppleProResMedia/Private/AppleProResEncoderProtocol.cpp b/Engine/Plugins/Media/AppleProResMedia/Source/AppleProResMedia/Private/AppleProResEncoderProtocol.cpp
--- a/Engine/Plugins/Media/AppleProResMedia/Source/AppleProResMedia/Private/AppleProResEncoderProtocol.cpp
+++ b/Engine/Plugins/Media/AppleProResMedia/Source/AppleProResMedia/Private/AppleProResEncoderProtocol.cpp
@@ -3,13 +3,10 @@
 #include "AppleProResEncoderProtocol.h"
 
 #include "AppleProResMediaModule.h"
-
 #include "Engine/RendererSettings.h"
-
 #include "HAL/FileManager.h"
-
 #include "MediaShaders.h"
-
+#include "Modules/ModuleManager.h"
 #include "MovieSceneCaptureProtocolBase.h"
 
 UAppleProResEncoderProtocol::UAppleProResEncoderProtocol(const FObjectInitializer& InObjInit)
@@ -97,6 +94,16 @@ bool UAppleProResEncoderProtocol::CreateProResFile(const FString& InFilename)
 		return false;
 	}
 
+	// Manually loading the toolbox library, on demand only.
+	if (IAppleProResMediaModule* AppleProResMediaModule = FModuleManager::GetModulePtr<IAppleProResMediaModule>("AppleProResMedia"))
+	{
+		if (!AppleProResMediaModule->LoadProResToolbox())
+		{
+			UE_LOG(LogAppleProResMedia, Error, TEXT("Failed to initialize ProRes Encoder."));
+			return false;	
+		}
+	}
+
 	const auto Filename = StringCast<ANSICHAR>(*FPaths::ConvertRelativePathToFull(InFilename));
 
 	Encoder = PROpenEncoder(NumberOfEncodingThreads, nullptr);
diff --git a/Engine/Plugins/Media/AppleProResMedia/Source/AppleProResMedia/Private/AppleProResMediaModule.cpp b/Engine/Plugins/Media/AppleProResMedia/Source/AppleProResMedia/Private/AppleProResMediaModule.cpp
--- a/Engine/Plugins/Media/AppleProResMedia/Source/AppleProResMedia/Private/AppleProResMediaModule.cpp
+++ b/Engine/Plugins/Media/AppleProResMedia/Source/AppleProResMedia/Private/AppleProResMediaModule.cpp
@@ -31,33 +31,45 @@
 
 DEFINE_LOG_CATEGORY(LogAppleProResMedia);
 
-class FAppleProResMediaModule : public IModuleInterface
+class FAppleProResMediaModule : public IAppleProResMediaModule
 {
-	virtual void StartupModule() override
+	virtual bool LoadProResToolbox() override
 	{
 #if PLATFORM_WINDOWS
-		check(LibHandle == nullptr);
-
-		const FString ProResDll = TEXT("ProResToolbox.dll");
+		if (ProResToolboxLibHandle)
+		{
+			return true;
+		}
+	
+		const FString ProResDllName = TEXT("ProResToolbox.dll");
 
 		// determine directory paths
 		FString ProResDllPath = FPaths::Combine(IPluginManager::Get().FindPlugin(TEXT("AppleProResMedia"))->GetBaseDir(), TEXT("/Binaries/ThirdParty/Win64"));
 		FPlatformProcess::PushDllDirectory(*ProResDllPath);
-		ProResDllPath = FPaths::Combine(ProResDllPath, ProResDll);
+		
+		ProResDllPath = FPaths::Combine(ProResDllPath, ProResDllName);
 
 		if (!FPaths::FileExists(ProResDllPath))
 		{
-			UE_LOG(LogAppleProResMedia, Error, TEXT("Failed to find the binary folder for the ProRes dll. Plug-in will not be functional."));
-			return;
+			UE_LOG(LogAppleProResMedia, Error, TEXT("Failed to find the binary folder for \"%s\". Plug-in will not be fully functional."), *ProResDllName);
+			return false;
 		}
 
-		LibHandle = FPlatformProcess::GetDllHandle(*ProResDllPath);
-
-		if (LibHandle == nullptr)
+		ProResToolboxLibHandle = FPlatformProcess::GetDllHandle(*ProResDllPath);
+		
+		if (ProResToolboxLibHandle == nullptr)
 		{
-			UE_LOG(LogAppleProResMedia, Error, TEXT("Failed to load required library %s. Plug-in will not be functional."), *ProResDllPath);
-			return;
+			UE_LOG(LogAppleProResMedia, Error, TEXT("Failed to load required library %s. Plug-in will not be fully functional."), *ProResDllPath);
+			return false;
 		}
+#endif
+		return true;
+	}
+	
+	virtual void StartupModule() override
+	{
+#if PLATFORM_WINDOWS
+		check(ProResToolboxLibHandle == nullptr);
 
 		if (IWmfMediaModule* Module = IWmfMediaModule::Get())
 		{
@@ -98,10 +110,10 @@ class FAppleProResMediaModule : public IModuleInterface
 		}
 #endif //WITH_EDITOR
 
-		if (LibHandle != nullptr)
+		if (ProResToolboxLibHandle != nullptr)
 		{
-			FPlatformProcess::FreeDllHandle(LibHandle);
-			LibHandle = nullptr;
+			FPlatformProcess::FreeDllHandle(ProResToolboxLibHandle);
+			ProResToolboxLibHandle = nullptr;
 		}
 	}
 
@@ -113,10 +125,10 @@ class FAppleProResMediaModule : public IModuleInterface
 
 private:
 
-	static void* LibHandle;
+	static void* ProResToolboxLibHandle;
 };
 
-void* FAppleProResMediaModule::LibHandle = nullptr;
+void* FAppleProResMediaModule::ProResToolboxLibHandle = nullptr;
 
 
 
diff --git a/Engine/Plugins/Media/AppleProResMedia/Source/AppleProResMedia/Private/AppleProResMediaModule.h b/Engine/Plugins/Media/AppleProResMedia/Source/AppleProResMedia/Private/AppleProResMediaModule.h
--- a/Engine/Plugins/Media/AppleProResMedia/Source/AppleProResMedia/Private/AppleProResMediaModule.h
+++ b/Engine/Plugins/Media/AppleProResMedia/Source/AppleProResMedia/Private/AppleProResMediaModule.h
@@ -3,6 +3,7 @@
 #pragma once
 
 #include "CoreMinimal.h"
+#include "Modules\ModuleInterface.h"
 
 #if PLATFORM_MAC
 #include <libkern/OSByteOrder.h>
@@ -10,3 +11,16 @@
 #endif
 
 DECLARE_LOG_CATEGORY_EXTERN(LogAppleProResMedia, Log, All);
+
+/**
+ * AppleProResMedia module interface
+ */
+class IAppleProResMediaModule : public IModuleInterface
+{
+public:
+	/**
+	 * ProResToolbox is delay loaded but needs to be explicitly loaded.
+	 * @return true if the library was loaded, false otherwise.
+	 */
+	virtual bool LoadProResToolbox() = 0;
+};
