[Backout] - CL46266885
#rnx
Original CL Desc
-----------------------------------------------------------------
[AppleProResMedia] Editor load time improvement: Implementing on demand loading of the ProResToolbox only when the encoder is used.

There is a known issue with this dll on windows. If a debugger is attached to the process, it will take between 20 to 40 seconds to load (cause still under investigation).

By making it load on demand, we are at least mitigating the issue for everyone that are not using the encoder.

diff --git a/Engine/Plugins/Media/AppleProResMedia/Source/AppleProResMedia/Private/AppleProResEncoder/AppleProResEncoder.cpp b/Engine/Plugins/Media/AppleProResMedia/Source/AppleProResMedia/Private/AppleProResEncoder/AppleProResEncoder.cpp
--- a/Engine/Plugins/Media/AppleProResMedia/Source/AppleProResMedia/Private/AppleProResEncoder/AppleProResEncoder.cpp
+++ b/Engine/Plugins/Media/AppleProResMedia/Source/AppleProResMedia/Private/AppleProResEncoder/AppleProResEncoder.cpp
@@ -1,9 +1,8 @@
 // Copyright Epic Games, Inc. All Rights Reserved.
 
 #include "AppleProResEncoder/AppleProResEncoder.h"
-#include "AppleProResMediaModule.h"
 #include "HAL/FileManager.h"
-#include "Modules/ModuleManager.h"
+#include "AppleProResMediaModule.h"
 #include "ProResToolboxExtension.h"
 
 namespace AppleProRes
@@ -277,16 +276,6 @@ FAppleProResEncoder::~FAppleProResEncoder()
 
 bool FAppleProResEncoder::Initialize()
 {
-	// Manually loading the toolbox library, on demand only.
-	if (IAppleProResMediaModule* AppleProResMediaModule = FModuleManager::GetModulePtr<IAppleProResMediaModule>("AppleProResMedia"))
-	{
-		if (!AppleProResMediaModule->LoadProResToolbox())
-		{
-			UE_LOG(LogAppleProResMedia, Error, TEXT("Failed to initialize ProRes Encoder."));
-			return false;	
-		}
-	}
-	
 	// Attempt to create the Encoder
 	Encoder = PROpenEncoder(FMath::Max(Options.MaxNumberOfEncodingThreads, 0), /*ThreadStartupCallback*/ nullptr);
 	if (!Encoder)
diff --git a/Engine/Plugins/Media/AppleProResMedia/Source/AppleProResMedia/Private/AppleProResEncoderProtocol.cpp b/Engine/Plugins/Media/AppleProResMedia/Source/AppleProResMedia/Private/AppleProResEncoderProtocol.cpp
--- a/Engine/Plugins/Media/AppleProResMedia/Source/AppleProResMedia/Private/AppleProResEncoderProtocol.cpp
+++ b/Engine/Plugins/Media/AppleProResMedia/Source/AppleProResMedia/Private/AppleProResEncoderProtocol.cpp
@@ -3,10 +3,13 @@
 #include "AppleProResEncoderProtocol.h"
 
 #include "AppleProResMediaModule.h"
+
 #include "Engine/RendererSettings.h"
+
 #include "HAL/FileManager.h"
+
 #include "MediaShaders.h"
-#include "Modules/ModuleManager.h"
+
 #include "MovieSceneCaptureProtocolBase.h"
 
 UAppleProResEncoderProtocol::UAppleProResEncoderProtocol(const FObjectInitializer& InObjInit)
@@ -94,16 +97,6 @@ bool UAppleProResEncoderProtocol::CreateProResFile(const FString& InFilename)
 		return false;
 	}
 
-	// Manually loading the toolbox library, on demand only.
-	if (IAppleProResMediaModule* AppleProResMediaModule = FModuleManager::GetModulePtr<IAppleProResMediaModule>("AppleProResMedia"))
-	{
-		if (!AppleProResMediaModule->LoadProResToolbox())
-		{
-			UE_LOG(LogAppleProResMedia, Error, TEXT("Failed to initialize ProRes Encoder."));
-			return false;	
-		}
-	}
-
 	const auto Filename = StringCast<ANSICHAR>(*FPaths::ConvertRelativePathToFull(InFilename));
 
 	Encoder = PROpenEncoder(NumberOfEncodingThreads, nullptr);
diff --git a/Engine/Plugins/Media/AppleProResMedia/Source/AppleProResMedia/Private/AppleProResMediaModule.cpp b/Engine/Plugins/Media/AppleProResMedia/Source/AppleProResMedia/Private/AppleProResMediaModule.cpp
--- a/Engine/Plugins/Media/AppleProResMedia/Source/AppleProResMedia/Private/AppleProResMediaModule.cpp
+++ b/Engine/Plugins/Media/AppleProResMedia/Source/AppleProResMedia/Private/AppleProResMediaModule.cpp
@@ -31,45 +31,33 @@
 
 DEFINE_LOG_CATEGORY(LogAppleProResMedia);
 
-class FAppleProResMediaModule : public IAppleProResMediaModule
+class FAppleProResMediaModule : public IModuleInterface
 {
-	virtual bool LoadProResToolbox() override
+	virtual void StartupModule() override
 	{
 #if PLATFORM_WINDOWS
-		if (ProResToolboxLibHandle)
-		{
-			return true;
-		}
-	
-		const FString ProResDllName = TEXT("ProResToolbox.dll");
+		check(LibHandle == nullptr);
+
+		const FString ProResDll = TEXT("ProResToolbox.dll");
 
 		// determine directory paths
 		FString ProResDllPath = FPaths::Combine(IPluginManager::Get().FindPlugin(TEXT("AppleProResMedia"))->GetBaseDir(), TEXT("/Binaries/ThirdParty/Win64"));
 		FPlatformProcess::PushDllDirectory(*ProResDllPath);
-		
-		ProResDllPath = FPaths::Combine(ProResDllPath, ProResDllName);
+		ProResDllPath = FPaths::Combine(ProResDllPath, ProResDll);
 
 		if (!FPaths::FileExists(ProResDllPath))
 		{
-			UE_LOG(LogAppleProResMedia, Error, TEXT("Failed to find the binary folder for \"%s\". Plug-in will not be fully functional."), *ProResDllName);
-			return false;
+			UE_LOG(LogAppleProResMedia, Error, TEXT("Failed to find the binary folder for the ProRes dll. Plug-in will not be functional."));
+			return;
 		}
 
-		ProResToolboxLibHandle = FPlatformProcess::GetDllHandle(*ProResDllPath);
-		
-		if (ProResToolboxLibHandle == nullptr)
+		LibHandle = FPlatformProcess::GetDllHandle(*ProResDllPath);
+
+		if (LibHandle == nullptr)
 		{
-			UE_LOG(LogAppleProResMedia, Error, TEXT("Failed to load required library %s. Plug-in will not be fully functional."), *ProResDllPath);
-			return false;
+			UE_LOG(LogAppleProResMedia, Error, TEXT("Failed to load required library %s. Plug-in will not be functional."), *ProResDllPath);
+			return;
 		}
-#endif
-		return true;
-	}
-	
-	virtual void StartupModule() override
-	{
-#if PLATFORM_WINDOWS
-		check(ProResToolboxLibHandle == nullptr);
 
 		if (IWmfMediaModule* Module = IWmfMediaModule::Get())
 		{
@@ -110,10 +98,10 @@ class FAppleProResMediaModule : public IAppleProResMediaModule
 		}
 #endif //WITH_EDITOR
 
-		if (ProResToolboxLibHandle != nullptr)
+		if (LibHandle != nullptr)
 		{
-			FPlatformProcess::FreeDllHandle(ProResToolboxLibHandle);
-			ProResToolboxLibHandle = nullptr;
+			FPlatformProcess::FreeDllHandle(LibHandle);
+			LibHandle = nullptr;
 		}
 	}
 
@@ -125,10 +113,10 @@ class FAppleProResMediaModule : public IAppleProResMediaModule
 
 private:
 
-	static void* ProResToolboxLibHandle;
+	static void* LibHandle;
 };
 
-void* FAppleProResMediaModule::ProResToolboxLibHandle = nullptr;
+void* FAppleProResMediaModule::LibHandle = nullptr;
 
 
 
diff --git a/Engine/Plugins/Media/AppleProResMedia/Source/AppleProResMedia/Private/AppleProResMediaModule.h b/Engine/Plugins/Media/AppleProResMedia/Source/AppleProResMedia/Private/AppleProResMediaModule.h
--- a/Engine/Plugins/Media/AppleProResMedia/Source/AppleProResMedia/Private/AppleProResMediaModule.h
+++ b/Engine/Plugins/Media/AppleProResMedia/Source/AppleProResMedia/Private/AppleProResMediaModule.h
@@ -3,7 +3,6 @@
 #pragma once
 
 #include "CoreMinimal.h"
-#include "Modules\ModuleInterface.h"
 
 #if PLATFORM_MAC
 #include <libkern/OSByteOrder.h>
@@ -11,16 +10,3 @@
 #endif
 
 DECLARE_LOG_CATEGORY_EXTERN(LogAppleProResMedia, Log, All);
-
-/**
- * AppleProResMedia module interface
- */
-class IAppleProResMediaModule : public IModuleInterface
-{
-public:
-	/**
-	 * ProResToolbox is delay loaded but needs to be explicitly loaded.
-	 * @return true if the library was loaded, false otherwise.
-	 */
-	virtual bool LoadProResToolbox() = 0;
-};
