[AppleProResMedia] Editor load time improvement: Implementing on demand loading of the ProResToolbox only when the encoder is used.

There is a known issue with this dll on windows. If a debugger is attached to the process, it will take between 20 to 40 seconds to load (cause still under investigation).

By making it load on demand, we are at least mitigating the issue for everyone that are not using the encoder.

#was-backout 46267895

diff --git a/Engine/Plugins/Media/AppleProResMedia/Source/AppleProResMedia/Private/AppleProResEncoder/AppleProResEncoder.cpp b/Engine/Plugins/Media/AppleProResMedia/Source/AppleProResMedia/Private/AppleProResEncoder/AppleProResEncoder.cpp
--- a/Engine/Plugins/Media/AppleProResMedia/Source/AppleProResMedia/Private/AppleProResEncoder/AppleProResEncoder.cpp
+++ b/Engine/Plugins/Media/AppleProResMedia/Source/AppleProResMedia/Private/AppleProResEncoder/AppleProResEncoder.cpp
@@ -1,8 +1,8 @@
 // Copyright Epic Games, Inc. All Rights Reserved.
 
 #include "AppleProResEncoder/AppleProResEncoder.h"
-#include "HAL/FileManager.h"
 #include "AppleProResMediaModule.h"
+#include "HAL/FileManager.h"
 #include "ProResToolboxExtension.h"
 
 namespace AppleProRes
@@ -276,6 +276,13 @@ FAppleProResEncoder::~FAppleProResEncoder()
 
 bool FAppleProResEncoder::Initialize()
 {
+	// Manually loading the toolbox library, on demand only.
+	if (!UE::AppleProResMedia::LoadProResToolbox())
+	{
+		UE_LOG(LogAppleProResMedia, Error, TEXT("Failed to initialize ProRes Encoder."));
+		return false;
+	}
+	
 	// Attempt to create the Encoder
 	Encoder = PROpenEncoder(FMath::Max(Options.MaxNumberOfEncodingThreads, 0), /*ThreadStartupCallback*/ nullptr);
 	if (!Encoder)
@@ -601,7 +608,7 @@ bool FAppleProResEncoder::WriteFrame(const FImagePixelData* InPixelData)
 {
 	check(InPixelData);
 
-	if (!ensureMsgf(bInitialized || bFinalized, TEXT("WriteFrame should not be called if not initialized or after finalize! Initialized: %d Finalized: %d"), bInitialized, bFinalized))
+	if (!ensureMsgf(bInitialized && !bFinalized, TEXT("WriteFrame should not be called if not initialized or after finalize! Initialized: %d Finalized: %d"), bInitialized, bFinalized))
 	{
 		return false;
 	}
@@ -734,7 +741,7 @@ bool FAppleProResEncoder::WriteAudioSample(const TArrayView<int16>& InAudioSampl
 		return true;
 	}
 
-	if (!ensureMsgf(bInitialized || bFinalized, TEXT("WriteFrame should not be called if not initialized or after finalize! Initialized: %d Finalized: %d"), bInitialized, bFinalized))
+	if (!ensureMsgf(bInitialized && !bFinalized, TEXT("WriteAudioSample should not be called if not initialized or after finalize! Initialized: %d Finalized: %d"), bInitialized, bFinalized))
 	{
 		return false;
 	}
diff --git a/Engine/Plugins/Media/AppleProResMedia/Source/AppleProResMedia/Private/AppleProResEncoderProtocol.cpp b/Engine/Plugins/Media/AppleProResMedia/Source/AppleProResMedia/Private/AppleProResEncoderProtocol.cpp
--- a/Engine/Plugins/Media/AppleProResMedia/Source/AppleProResMedia/Private/AppleProResEncoderProtocol.cpp
+++ b/Engine/Plugins/Media/AppleProResMedia/Source/AppleProResMedia/Private/AppleProResEncoderProtocol.cpp
@@ -3,13 +3,9 @@
 #include "AppleProResEncoderProtocol.h"
 
 #include "AppleProResMediaModule.h"
-
 #include "Engine/RendererSettings.h"
-
 #include "HAL/FileManager.h"
-
 #include "MediaShaders.h"
-
 #include "MovieSceneCaptureProtocolBase.h"
 
 UAppleProResEncoderProtocol::UAppleProResEncoderProtocol(const FObjectInitializer& InObjInit)
@@ -97,6 +93,13 @@ bool UAppleProResEncoderProtocol::CreateProResFile(const FString& InFilename)
 		return false;
 	}
 
+	// Manually loading the toolbox library, on demand only.
+	if (!UE::AppleProResMedia::LoadProResToolbox())
+	{
+		UE_LOG(LogAppleProResMedia, Error, TEXT("Failed to initialize ProRes Encoder."));
+		return false;
+	}
+
 	const auto Filename = StringCast<ANSICHAR>(*FPaths::ConvertRelativePathToFull(InFilename));
 
 	Encoder = PROpenEncoder(NumberOfEncodingThreads, nullptr);
diff --git a/Engine/Plugins/Media/AppleProResMedia/Source/AppleProResMedia/Private/AppleProResMediaModule.cpp b/Engine/Plugins/Media/AppleProResMedia/Source/AppleProResMedia/Private/AppleProResMediaModule.cpp
--- a/Engine/Plugins/Media/AppleProResMedia/Source/AppleProResMedia/Private/AppleProResMediaModule.cpp
+++ b/Engine/Plugins/Media/AppleProResMedia/Source/AppleProResMedia/Private/AppleProResMediaModule.cpp
@@ -33,32 +33,49 @@ DEFINE_LOG_CATEGORY(LogAppleProResMedia);
 
 class FAppleProResMediaModule : public IModuleInterface
 {
-	virtual void StartupModule() override
+public:
+	bool LoadProResToolbox()
 	{
 #if PLATFORM_WINDOWS
-		check(LibHandle == nullptr);
-
-		const FString ProResDll = TEXT("ProResToolbox.dll");
+		if (ProResToolboxLibHandle)
+		{
+			return true;
+		}
+	
+		const FString ProResDllName = TEXT("ProResToolbox.dll");
 
 		// determine directory paths
-		FString ProResDllPath = FPaths::Combine(IPluginManager::Get().FindPlugin(TEXT("AppleProResMedia"))->GetBaseDir(), TEXT("/Binaries/ThirdParty/Win64"));
-		FPlatformProcess::PushDllDirectory(*ProResDllPath);
-		ProResDllPath = FPaths::Combine(ProResDllPath, ProResDll);
 
-		if (!FPaths::FileExists(ProResDllPath))
+		TSharedPtr<IPlugin> ThisPlugin = IPluginManager::Get().FindPlugin(TEXT("AppleProResMedia"));
+		if (!ThisPlugin)
 		{
-			UE_LOG(LogAppleProResMedia, Error, TEXT("Failed to find the binary folder for the ProRes dll. Plug-in will not be functional."));
-			return;
+			UE_LOG(LogAppleProResMedia, Error, TEXT("Failed to find the plugin folder for \"%s\". Plug-in will not be fully functional."), *ProResDllName);
+			return false;
 		}
 
-		LibHandle = FPlatformProcess::GetDllHandle(*ProResDllPath);
-
-		if (LibHandle == nullptr)
+		const FString ProResDllPath = FPaths::Combine(ThisPlugin->GetBaseDir(), TEXT("Binaries"), TEXT("ThirdParty"), TEXT("Win64"));
+		const FString ProResDllFullPath = FPaths::Combine(ProResDllPath, ProResDllName);
+		
+		if (!FPaths::FileExists(ProResDllFullPath))
 		{
-			UE_LOG(LogAppleProResMedia, Error, TEXT("Failed to load required library %s. Plug-in will not be functional."), *ProResDllPath);
-			return;
+			UE_LOG(LogAppleProResMedia, Error, TEXT("Failed to find required library \"%s\". Plug-in will not be fully functional."), *ProResDllFullPath);
+			return false;
 		}
 
+		ProResToolboxLibHandle = FPlatformProcess::GetDllHandle(*ProResDllFullPath);
+		
+		if (ProResToolboxLibHandle == nullptr)
+		{
+			UE_LOG(LogAppleProResMedia, Error, TEXT("Failed to load required library \"%s\". Plug-in will not be fully functional."), *ProResDllFullPath);
+			return false;
+		}
+#endif
+		return true;
+	}
+	
+	virtual void StartupModule() override
+	{
+#if PLATFORM_WINDOWS
 		if (IWmfMediaModule* Module = IWmfMediaModule::Get())
 		{
 			if (Module->IsInitialized())
@@ -98,10 +115,10 @@ class FAppleProResMediaModule : public IModuleInterface
 		}
 #endif //WITH_EDITOR
 
-		if (LibHandle != nullptr)
+		if (ProResToolboxLibHandle != nullptr)
 		{
-			FPlatformProcess::FreeDllHandle(LibHandle);
-			LibHandle = nullptr;
+			FPlatformProcess::FreeDllHandle(ProResToolboxLibHandle);
+			ProResToolboxLibHandle = nullptr;
 		}
 	}
 
@@ -113,12 +130,23 @@ class FAppleProResMediaModule : public IModuleInterface
 
 private:
 
-	static void* LibHandle;
+	static void* ProResToolboxLibHandle;
 };
 
-void* FAppleProResMediaModule::LibHandle = nullptr;
-
+void* FAppleProResMediaModule::ProResToolboxLibHandle = nullptr;
 
+bool UE::AppleProResMedia::LoadProResToolbox()
+{
+#if PLATFORM_WINDOWS
+	if (FAppleProResMediaModule* AppleProResMediaModule = FModuleManager::GetModulePtr<FAppleProResMediaModule>("AppleProResMedia"))
+	{
+		return AppleProResMediaModule->LoadProResToolbox();
+	}
+	return false;
+#else
+	return true;	// There is nothing to do for other platforms
+#endif
+}
 
 IMPLEMENT_MODULE(FAppleProResMediaModule, AppleProResMedia)
 
diff --git a/Engine/Plugins/Media/AppleProResMedia/Source/AppleProResMedia/Private/AppleProResMediaModule.h b/Engine/Plugins/Media/AppleProResMedia/Source/AppleProResMedia/Private/AppleProResMediaModule.h
--- a/Engine/Plugins/Media/AppleProResMedia/Source/AppleProResMedia/Private/AppleProResMediaModule.h
+++ b/Engine/Plugins/Media/AppleProResMedia/Source/AppleProResMedia/Private/AppleProResMediaModule.h
@@ -10,3 +10,12 @@
 #endif
 
 DECLARE_LOG_CATEGORY_EXTERN(LogAppleProResMedia, Log, All);
+
+namespace UE::AppleProResMedia
+{
+	/**
+	 * ProResToolbox is delay loaded but needs to be manually loaded because of search path limitations.
+	 * @return true if the library was loaded, false otherwise.
+	 */
+	bool LoadProResToolbox();
+}
\ No newline at end of file
