[AppleProResMedia] Fix Error handling and potential memory leaks.

diff --git a/Engine/Plugins/Media/AppleProResMedia/Source/AppleProResMedia/Private/AppleProResEncoder/AppleProResEncoder.cpp b/Engine/Plugins/Media/AppleProResMedia/Source/AppleProResMedia/Private/AppleProResEncoder/AppleProResEncoder.cpp
--- a/Engine/Plugins/Media/AppleProResMedia/Source/AppleProResMedia/Private/AppleProResEncoder/AppleProResEncoder.cpp
+++ b/Engine/Plugins/Media/AppleProResMedia/Source/AppleProResMedia/Private/AppleProResEncoder/AppleProResEncoder.cpp
@@ -3,6 +3,7 @@
 #include "AppleProResEncoder/AppleProResEncoder.h"
 #include "AppleProResMediaModule.h"
 #include "HAL/FileManager.h"
+#include "Misc/ScopeExit.h"
 #include "ProResToolboxExtension.h"
 
 namespace AppleProRes
@@ -291,6 +292,15 @@ bool FAppleProResEncoder::Initialize()
 		return false;
 	}
 
+	ON_SCOPE_EXIT
+	{
+		// Clear any allocated resources if initialization failed.
+		if (!bInitialized)
+		{
+			ReleaseAndClearResources();
+		}
+	};
+
 	// Delete the file on disk if it already exists as ProRes doesn't handle that for us.
 	if (IFileManager::Get().FileExists(*Options.OutputFilename))
 	{
@@ -542,10 +552,10 @@ bool FAppleProResEncoder::InitializeAudioTrack()
 
 void FAppleProResEncoder::Finalize()
 {
-    if(bFinalized || !bInitialized)
-    {
-        return;
-    }
+	if (bFinalized || !bInitialized)
+	{
+		return;
+	}
 
 	if (FileWriter)
 	{
@@ -589,21 +599,30 @@ void FAppleProResEncoder::Finalize()
 		{
 			UE_LOG(LogAppleProResMedia, Error, TEXT("Failed to Invalidate Writer for file writer at path: \"%s\"."), *Options.OutputFilename);
 		}
-
-		// Close the Encoder
-		PRCloseEncoder(Encoder);
 	}
 
 	// Release resources allocated during initialization
+	ReleaseAndClearResources();
+
+	bFinalized = true;
+}
+
+void FAppleProResEncoder::ReleaseAndClearResources()
+{
 	PRReleaseAndClear(&VideoFormatDescription);
 	PRReleaseAndClear(&AudioFormatDescription);
 	PRReleaseAndClear(&TimecodeFormatDescription);
 	PRReleaseAndClear(&FileWriter);
-    
-    bFinalized = true;
+
+	if (Encoder)
+	{
+		PRCloseEncoder(Encoder);
+		Encoder = nullptr;
+	}
+
+	bInitialized = false;
 }
-    
-    
+
 bool FAppleProResEncoder::WriteFrame(const FImagePixelData* InPixelData)
 {
 	check(InPixelData);
@@ -746,6 +765,16 @@ bool FAppleProResEncoder::WriteAudioSample(const TArrayView<int16>& InAudioSampl
 		return false;
 	}
 
+	if (!ensureMsgf(Options.AudioChannelCount != 0, TEXT("AudioChannelCount can't be zero.")))
+	{
+		return false;
+	}
+
+	if (!ensureMsgf(AudioSamplesPerFrame != 0, TEXT("AudioSamplesPerFrame can't be zero.")))
+	{
+		return false;
+	}
+
 	int32 SampleCountPerChannel = InAudioSamples.Num() / Options.AudioChannelCount;
 	int32 NumFramesOfData = SampleCountPerChannel / AudioSamplesPerFrame;
 
diff --git a/Engine/Plugins/Media/AppleProResMedia/Source/AppleProResMedia/Private/AppleProResEncoderProtocol.cpp b/Engine/Plugins/Media/AppleProResMedia/Source/AppleProResMedia/Private/AppleProResEncoderProtocol.cpp
--- a/Engine/Plugins/Media/AppleProResMedia/Source/AppleProResMedia/Private/AppleProResEncoderProtocol.cpp
+++ b/Engine/Plugins/Media/AppleProResMedia/Source/AppleProResMedia/Private/AppleProResEncoderProtocol.cpp
@@ -6,6 +6,7 @@
 #include "Engine/RendererSettings.h"
 #include "HAL/FileManager.h"
 #include "MediaShaders.h"
+#include "Misc/ScopeExit.h"
 #include "MovieSceneCaptureProtocolBase.h"
 
 UAppleProResEncoderProtocol::UAppleProResEncoderProtocol(const FObjectInitializer& InObjInit)
@@ -107,6 +108,15 @@ bool UAppleProResEncoderProtocol::CreateProResFile(const FString& InFilename)
 	PRStatus status = 0;
 	status = ProResFileWriterCreate(Filename.Get(), &FileWriter);
 
+	ON_SCOPE_EXIT
+	{
+		// Clear any allocated resources if initialization failed.
+		if (status != 0)
+		{
+			ReleaseAndClearResources();
+		}
+	};
+
 	if (status != 0)
 	{
 		return false;
@@ -347,14 +357,9 @@ void UAppleProResEncoderProtocol::FinalizeImpl()
 		CurrentTime.value += FrameRate.Denominator;
 		status = ProResFileWriterEndSession(FileWriter, CurrentTime);
 		status = ProResFileWriterFinish(FileWriter);
-		PRRelease(FormatDescription);
-		if (bEmbedTimecodeTrack)
-		{
-			PRRelease(TimecodeFormatDescription);
-		}
-		PRRelease(FileWriter);
-		PRCloseEncoder(Encoder);
 	}
+
+	ReleaseAndClearResources();
 }
 
 bool UAppleProResEncoderProtocol::CanWriteToFileImpl(const TCHAR* InFilename, bool bInOverwriteExisting) const
@@ -595,3 +600,16 @@ void UAppleProResEncoderProtocol::ConvertFColorToRGBA4444(const TArray<FColor>&
 		SourceBuffer += 4;
 	}
 }
+
+void UAppleProResEncoderProtocol::ReleaseAndClearResources()
+{
+	PRReleaseAndClear(&FormatDescription);
+	PRReleaseAndClear(&TimecodeFormatDescription);
+	PRReleaseAndClear(&FileWriter);
+
+	if (Encoder)
+	{
+		PRCloseEncoder(Encoder);
+		Encoder = nullptr;
+	}
+}
\ No newline at end of file
diff --git a/Engine/Plugins/Media/AppleProResMedia/Source/AppleProResMedia/Private/AppleProResEncoderProtocol.h b/Engine/Plugins/Media/AppleProResMedia/Source/AppleProResMedia/Private/AppleProResEncoderProtocol.h
--- a/Engine/Plugins/Media/AppleProResMedia/Source/AppleProResMedia/Private/AppleProResEncoderProtocol.h
+++ b/Engine/Plugins/Media/AppleProResMedia/Source/AppleProResMedia/Private/AppleProResEncoderProtocol.h
@@ -104,6 +104,7 @@ private:
 	uint32_t GetFieldCount() const;
 	ProResFormatDescriptionFieldDetail GetFieldDetail() const;
 	bool HasAlpha() const;
+	void ReleaseAndClearResources();
 
 private:
 
diff --git a/Engine/Plugins/Media/AppleProResMedia/Source/AppleProResMedia/Private/AppleProResMediaModule.cpp b/Engine/Plugins/Media/AppleProResMedia/Source/AppleProResMedia/Private/AppleProResMediaModule.cpp
--- a/Engine/Plugins/Media/AppleProResMedia/Source/AppleProResMedia/Private/AppleProResMediaModule.cpp
+++ b/Engine/Plugins/Media/AppleProResMedia/Source/AppleProResMedia/Private/AppleProResMediaModule.cpp
@@ -62,7 +62,9 @@ public:
 			return false;
 		}
 
+		FPlatformProcess::PushDllDirectory(*ProResDllPath);
 		ProResToolboxLibHandle = FPlatformProcess::GetDllHandle(*ProResDllFullPath);
+		FPlatformProcess::PopDllDirectory(*ProResDllPath);
 		
 		if (ProResToolboxLibHandle == nullptr)
 		{
diff --git a/Engine/Plugins/Media/AppleProResMedia/Source/AppleProResMedia/Public/AppleProResEncoder/AppleProResEncoder.h b/Engine/Plugins/Media/AppleProResMedia/Source/AppleProResMedia/Public/AppleProResEncoder/AppleProResEncoder.h
--- a/Engine/Plugins/Media/AppleProResMedia/Source/AppleProResMedia/Public/AppleProResEncoder/AppleProResEncoder.h
+++ b/Engine/Plugins/Media/AppleProResMedia/Source/AppleProResMedia/Public/AppleProResEncoder/AppleProResEncoder.h
@@ -164,6 +164,7 @@ private:
 	bool InitializeVideoTrack();
 	bool InitializeTimecodeTrack();
 	bool InitializeAudioTrack();
+	void ReleaseAndClearResources();
 
 private:
 	FAppleProResEncoderOptions Options;
